AIA Fix build error with clang

diff --git a/Engine/Plugins/Experimental/AIAssistant/Source/AIAssistant/Private/AIAssistantEnum.h b/Engine/Plugins/Experimental/AIAssistant/Source/AIAssistant/Private/AIAssistantEnum.h
--- a/Engine/Plugins/Experimental/AIAssistant/Source/AIAssistant/Private/AIAssistantEnum.h
+++ b/Engine/Plugins/Experimental/AIAssistant/Source/AIAssistant/Private/AIAssistantEnum.h
@@ -46,7 +46,7 @@ namespace UE::AIAssistant
 	//	 X(Goodbye, "See-ya")
 	//
 	// The following generates a count of enum values stored in MyEnumCount.
-	// constexpr std::size_t MyEnumCount =
+	// constexpr uint32 MyEnumCount =
 	//     UE_ENUM_COUNT(MY_ENUM(UE_ENUM_COUNTER));
 	//
 	// The macro that defines an enum must *not* include a trailing comma.
@@ -79,7 +79,7 @@ namespace UE::AIAssistant
 	// * FString LexToString(EnumType Value) to convert from an enum value to string.
 	#define UE_ENUM_METADATA_DECLARE(EnumType, EnumDescriptionMacro)                        \
 		/* Number of members of EnumType */                                                 \
-		extern const std::size_t EnumType##Count;                                           \
+		extern const uint32 EnumType##Count;                                           \
 		                                                                                    \
 		/* Description of each member of EnumType */                                        \
 		extern const UE_ENUM_VALUE_DESCRIPTION_TYPE(EnumType, EnumDescriptionMacro)         \
@@ -93,7 +93,7 @@ namespace UE::AIAssistant
 
 	// Defines metadata for an enum in the current namespace. (use in a source file)
 	#define UE_ENUM_METADATA_DEFINE(EnumType, EnumDescriptionMacro)                         \
-		const std::size_t EnumType##Count =                                                 \
+		const uint32 EnumType##Count =                                                 \
 			UE_ENUM_COUNT(EnumDescriptionMacro(UE_ENUM_COUNTER));                           \
 		                                                                                    \
 		const UE_ENUM_VALUE_DESCRIPTION_TYPE(EnumType, EnumDescriptionMacro)                \
@@ -111,8 +111,8 @@ namespace UE::AIAssistant
 		                                                                                    \
 		FString LexToString(EnumType Value)                                                 \
 		{                                                                                   \
-			auto MaybeDescription =                                                         \
-				UE::AIAssistant::GetEnumValueDescription(EnumType##Descriptions, Value);    \
+			auto MaybeDescription = UE::AIAssistant::GetEnumValueDescription(                \
+				EnumType##Descriptions, Value);                                             \
 			return MaybeDescription.IsSet() ? **MaybeDescription : FString();               \
 		}
 
@@ -125,7 +125,7 @@ namespace UE::AIAssistant
 	};
 
 	// Get the description of an enum value.
-	template <typename T, std::size_t Size>
+	template <typename T, auto Size>
 	TOptional<const FString*> GetEnumValueDescription(
 		const TStaticArray<EnumValueDescription<T>, Size>& EnumValueDescriptions, T EnumValue)
 	{
@@ -142,7 +142,7 @@ namespace UE::AIAssistant
 	}
 
 	// Convert an enum description to an enum value returning no value if the description is not found.
-	template <typename T, std::size_t Size>
+	template <typename T, auto Size>
 	TOptional<T> GetEnumValueFromDescription(
 		const TStaticArray<EnumValueDescription<T>, Size>& EnumValueDescriptions,
 		const FString& Description, ESearchCase::Type SearchCase = ESearchCase::CaseSensitive)
diff --git a/Engine/Plugins/Experimental/AIAssistant/Source/AIAssistant/Private/AIAssistantJsonVariantSerializer.h b/Engine/Plugins/Experimental/AIAssistant/Source/AIAssistant/Private/AIAssistantJsonVariantSerializer.h
--- a/Engine/Plugins/Experimental/AIAssistant/Source/AIAssistant/Private/AIAssistantJsonVariantSerializer.h
+++ b/Engine/Plugins/Experimental/AIAssistant/Source/AIAssistant/Private/AIAssistantJsonVariantSerializer.h
@@ -81,8 +81,8 @@ namespace UE::AIAssistant
 					auto VariantJsonObj = Serializer.GetObject()->GetObjectField(VariantFieldName);
 					if (VariantJsonObj.IsValid())
 					{
-						VariantField.Emplace<T>();
-						VariantField.Get<T>().FromJson(VariantJsonObj);
+						VariantField.template Emplace<T>();
+						VariantField.template Get<T>().FromJson(VariantJsonObj);
 					}
 					else
 					{
@@ -94,12 +94,12 @@ namespace UE::AIAssistant
 			}
 			else
 			{
-				if (VariantField.IsType<T>())
+				if (VariantField.template IsType<T>())
 				{
 					FString VariantTypeJson = LexToString(VariantTypeField);
 					Serializer.Serialize(VariantTypeFieldName, VariantTypeJson);
 					Serializer.StartObject(VariantFieldName);
-					VariantField.Get<T>().Serialize(Serializer, true);
+					VariantField.template Get<T>().Serialize(Serializer, true);
 					Serializer.EndObject();
 				}
 			}
diff --git a/Engine/Plugins/Experimental/AIAssistant/Source/AIAssistant/Private/AIAssistantWebJavaScriptResultDelegate.cpp b/Engine/Plugins/Experimental/AIAssistant/Source/AIAssistant/Private/AIAssistantWebJavaScriptResultDelegate.cpp
--- a/Engine/Plugins/Experimental/AIAssistant/Source/AIAssistant/Private/AIAssistantWebJavaScriptResultDelegate.cpp
+++ b/Engine/Plugins/Experimental/AIAssistant/Source/AIAssistant/Private/AIAssistantWebJavaScriptResultDelegate.cpp
@@ -71,7 +71,7 @@ void UAIAssistantWebJavaScriptResultDelegate::HandleResult(
 		Handler = *HandlerPtrPtr;
 	}
 	bool RemoveHandler =
-		(*Handler)(FResultHandlerContext{ ResultJson, bResultJsonIsError, HandlerId });
+		(*Handler)(FResultHandlerContext{ { ResultJson, bResultJsonIsError }, HandlerId });
 	if (RemoveHandler)
 	{
 		UE::TUniqueLock Lock(ResultHandlersByIdLock);
diff --git a/Engine/Plugins/Experimental/AIAssistant/Source/AIAssistant/Private/Tests/AIAssistantEnumTest.cpp b/Engine/Plugins/Experimental/AIAssistant/Source/AIAssistant/Private/Tests/AIAssistantEnumTest.cpp
--- a/Engine/Plugins/Experimental/AIAssistant/Source/AIAssistant/Private/Tests/AIAssistantEnumTest.cpp
+++ b/Engine/Plugins/Experimental/AIAssistant/Source/AIAssistant/Private/Tests/AIAssistantEnumTest.cpp
@@ -24,7 +24,7 @@ enum class EAdjective
 
 static const EAdjective Adjectives[] = { ADJECTIVE_ENUM(UE_ENUM_VALUE) };
 
-static constexpr std::size_t AdjectivesCount =
+static constexpr uint32 AdjectivesCount =
   UE_ENUM_COUNT(ADJECTIVE_ENUM(UE_ENUM_COUNTER));
 
 static const UE_ENUM_VALUE_DESCRIPTION_TYPE(EAdjective, ADJECTIVE_ENUM)
