[AI Assistant] Fixed some AI assistant docs issues

diff --git a/Engine/Plugins/Experimental/AIAssistant/Source/AIAssistant/Private/AIAssistantConfig.h b/Engine/Plugins/Experimental/AIAssistant/Source/AIAssistant/Private/AIAssistantConfig.h
--- a/Engine/Plugins/Experimental/AIAssistant/Source/AIAssistant/Private/AIAssistantConfig.h
+++ b/Engine/Plugins/Experimental/AIAssistant/Source/AIAssistant/Private/AIAssistantConfig.h
@@ -20,7 +20,7 @@ struct FAIAssistantConfig : public FJsonSerializable
 	TArray<FString> AllowedUrlRegexes;
 
 	// Allowed URL patterns (regular expressions) as strings including defaults.
-	// This contains MainUrl escaped as a regex, DefaultAllowedRegexes and AllowedUrlRegexes.
+	// This contains MainUrl escaped as a regex, DefaultAllowedUrlRegexes and AllowedUrlRegexes.
 	TArray<FString> GetAllAllowedUrlRegexes() const;
 
 	// Allowed URL patterns as regular expressions.
diff --git a/Engine/Plugins/Experimental/AIAssistant/Source/AIAssistant/Private/AIAssistantEnum.h b/Engine/Plugins/Experimental/AIAssistant/Source/AIAssistant/Private/AIAssistantEnum.h
--- a/Engine/Plugins/Experimental/AIAssistant/Source/AIAssistant/Private/AIAssistantEnum.h
+++ b/Engine/Plugins/Experimental/AIAssistant/Source/AIAssistant/Private/AIAssistantEnum.h
@@ -64,7 +64,7 @@ namespace UE::AIAssistant
 	//     X(Goodbye, "See-ya")
 	//
 	// The following declares an EnumValueDescription array.
-	// extern const UE_ENUM_VALUE_DESCRIPTION_TYPE(EMyEnum, MY_ENUM MyEnumDescription;
+	// extern const UE_ENUM_VALUE_DESCRIPTION_TYPE(EMyEnum, MY_ENUM MyEnumDescription);
 	#define UE_ENUM_VALUE_DESCRIPTION_TYPE(EnumType, EnumTypeMacro) \
 		TStaticArray< \
 			UE::AIAssistant::EnumValueDescription<EnumType>, \
