[AI Assistant] Fixed potentially uninitialized member of message content

diff --git a/Engine/Plugins/Experimental/AIAssistant/Source/AIAssistant/Private/AIAssistantWebApi.h b/Engine/Plugins/Experimental/AIAssistant/Source/AIAssistant/Private/AIAssistantWebApi.h
--- a/Engine/Plugins/Experimental/AIAssistant/Source/AIAssistant/Private/AIAssistantWebApi.h
+++ b/Engine/Plugins/Experimental/AIAssistant/Source/AIAssistant/Private/AIAssistantWebApi.h
@@ -66,7 +66,7 @@ namespace UE::AIAssistant
 		// Content of the message, this can be extended in future to support additional types of content.
 		TVariant<FTextMessageContent> Content;
 		// Whether the message is visible to the user.
-		bool bVisibleToUser;
+		bool bVisibleToUser = true;
 
 		BEGIN_JSON_SERIALIZER
 			UE_JSON_SERIALIZE_ENUM_VARIANT_BEGIN("contentType", ContentType, "content", Content);
