[AI Assistant] Fix docs bug in enum test

diff --git a/Engine/Plugins/Experimental/AIAssistant/Source/AIAssistant/Private/Tests/AIAssistantEnumTest.cpp b/Engine/Plugins/Experimental/AIAssistant/Source/AIAssistant/Private/Tests/AIAssistantEnumTest.cpp
--- a/Engine/Plugins/Experimental/AIAssistant/Source/AIAssistant/Private/Tests/AIAssistantEnumTest.cpp
+++ b/Engine/Plugins/Experimental/AIAssistant/Source/AIAssistant/Private/Tests/AIAssistantEnumTest.cpp
@@ -81,7 +81,7 @@ IMPLEMENT_SIMPLE_AUTOMATION_TEST(
 	"AI.Assistant.Enum.GetEnumValueDescriptionOutOfRange",
 	AIAssistantTest::Flags);
 
-// Ensure "Invalid" is returned for out of range enum queries.
+// Ensure no value is returned for out of range enum queries.
 bool FAIAssistantEnumTestGetEnumValueDescriptionOutOfRange::RunTest(const FString& UnusedParameters)
 {
 	(void)TestFalse(
@@ -117,7 +117,7 @@ IMPLEMENT_SIMPLE_AUTOMATION_TEST(
 	"AI.Assistant.Enum.GetEnumValueFromDescriptionInvalid",
 	AIAssistantTest::Flags);
 
-// Ensure the default value is returned for an invalid enum description.
+// Ensure no value is returned for an invalid enum description.
 bool FAIAssistantEnumTestGetEnumValueFromDescriptionInvalid::RunTest(const FString& UnusedParameters)
 {
 	(void)TestFalse(
