[AI Assistant] Override the default assistant config.

diff --git a/Engine/Plugins/Experimental/AIAssistant/Source/AIAssistant/Private/AIAssistantConfig.cpp b/Engine/Plugins/Experimental/AIAssistant/Source/AIAssistant/Private/AIAssistantConfig.cpp
--- a/Engine/Plugins/Experimental/AIAssistant/Source/AIAssistant/Private/AIAssistantConfig.cpp
+++ b/Engine/Plugins/Experimental/AIAssistant/Source/AIAssistant/Private/AIAssistantConfig.cpp
@@ -4,6 +4,7 @@
 
 #include "Containers/Set.h"
 #include "Misc/FileHelper.h"
+#include "Misc/Optional.h"
 #include "Misc/Paths.h"
 #include "Templates/UnrealTemplate.h"
 
@@ -86,11 +87,32 @@ TArray<FRegexPattern> FAIAssistantConfig::GetAllowedUrlRegexPatterns() const
 
 TArray<FString> FAIAssistantConfig::GetDefaultSearchDirectories()
 {
-	return {
+	TArray<FString> BaseSearchPaths =
+	{
 		FPaths::EngineConfigDir(),
 		FPaths::EngineUserDir(),
 		FPaths::EngineVersionAgnosticUserDir(),
 	};
+	TArray<FString> AllPathsToSearch;
+	TArray<TOptional<FPaths::EPathConversion>> PathConversions =
+	{
+		TOptional(FPaths::EPathConversion::Engine_NotForLicensees),
+		TOptional(FPaths::EPathConversion::Engine_NoRedist),
+		TOptional(FPaths::EPathConversion::Engine_LimitedAccess),
+		TOptional<FPaths::EPathConversion>(),
+	};
+	for (const TOptional<FPaths::EPathConversion>& PathConversion : PathConversions)
+	{
+		for (const FString& BaseSearchPath : BaseSearchPaths)
+		{
+			AllPathsToSearch.Emplace(
+				PathConversion.IsSet()
+				? FPaths::ConvertPath(BaseSearchPath, PathConversion.GetValue())
+				: BaseSearchPath);
+		}
+		
+	}
+	return AllPathsToSearch;
 }
 
 FString FAIAssistantConfig::FindConfigFile(const TArray<FString>& SearchDirectories)
diff --git a/Engine/Plugins/Experimental/AIAssistant/Source/AIAssistant/Private/AIAssistantConfig.h b/Engine/Plugins/Experimental/AIAssistant/Source/AIAssistant/Private/AIAssistantConfig.h
--- a/Engine/Plugins/Experimental/AIAssistant/Source/AIAssistant/Private/AIAssistantConfig.h
+++ b/Engine/Plugins/Experimental/AIAssistant/Source/AIAssistant/Private/AIAssistantConfig.h
@@ -34,7 +34,15 @@ struct FAIAssistantConfig : public FJsonSerializable
 	// Get the main URL as an escaped regex string.
 	FString GetMainUrlAsRegexString() const;
 
-	// Get the default configuration search directories:
+	// Get the default configuration search directories with the most specific overrides first.
+	//
+	// For example, when searching the engine configuration directory:
+	// * Engine/Restricted/NotForLicensees/Config/
+	// * Engine/Restricted/NoRedist/Config/
+	// * Engine/Restricted/LimitedAccess/Config/
+	// * Engine/Config/
+	//
+	// The base directories that are searched with overrides are:
 	// * Engine configuration directory: Engine/Config/
 	// * Engine directory: Engine/
 	// * Engine user directory (if installed): On Windows %LOCALAPPDATA%/UnrealEngine/Common/
