[AI Assistant] Initialize SWebBrowser::bOpenedUrlChanged

It's initialized in OnBeforeNavigation_Lambda before it's read but this makes sure that it is always initialized.

diff --git a/Engine/Plugins/Experimental/AIAssistant/Source/AIAssistant/Private/AIAssistantWebBrowser.h b/Engine/Plugins/Experimental/AIAssistant/Source/AIAssistant/Private/AIAssistantWebBrowser.h
--- a/Engine/Plugins/Experimental/AIAssistant/Source/AIAssistant/Private/AIAssistantWebBrowser.h
+++ b/Engine/Plugins/Experimental/AIAssistant/Source/AIAssistant/Private/AIAssistantWebBrowser.h
@@ -88,7 +88,7 @@ private:
 	// Previously opened URL.
 	FString LastOpenedUrl;
 	// Whether navigation changed.
-	bool bOpenedUrlChanged;
+	bool bOpenedUrlChanged = false;
 	// Previously loaded URL that was not a redirect.
 	FString LastOpenedNonRedirectUrl;
 
