Prevent the AI Assistant module from loading when running a commandlet

diff --git a/Engine/Plugins/Experimental/AIAssistant/AIAssistant.uplugin b/Engine/Plugins/Experimental/AIAssistant/AIAssistant.uplugin
--- a/Engine/Plugins/Experimental/AIAssistant/AIAssistant.uplugin
+++ b/Engine/Plugins/Experimental/AIAssistant/AIAssistant.uplugin
@@ -19,7 +19,7 @@
 	"Modules": [
 		{
 			"Name": "AIAssistant",
-			"Type": "Editor",
+			"Type": "EditorNoCommandlet",
 			"LoadingPhase": "Default"
 		}
 	],
