Fix crash when cooking with AIAssistant plugin enabled.

diff --git a/Engine/Plugins/Experimental/AIAssistant/Source/AIAssistant/Private/AIAssistant.cpp b/Engine/Plugins/Experimental/AIAssistant/Source/AIAssistant/Private/AIAssistant.cpp
--- a/Engine/Plugins/Experimental/AIAssistant/Source/AIAssistant/Private/AIAssistant.cpp
+++ b/Engine/Plugins/Experimental/AIAssistant/Source/AIAssistant/Private/AIAssistant.cpp
@@ -62,13 +62,18 @@ void FAIAssistantModule::StartupModule()
 
 	
 	InputProcessor = MakeShared<FAIAssistantInputProcessor>();
-	FSlateApplication::Get().RegisterInputPreProcessor(InputProcessor);
+
+	// Some scenarios, like commandlets, don't have slate initialized
+	if (FSlateApplication::IsInitialized())
+	{
+		FSlateApplication::Get().RegisterInputPreProcessor(InputProcessor);
+	}
 }
 
 
 void FAIAssistantModule::ShutdownModule()
 {
-	if (FSlateApplication::IsInitialized()) // ..prevents crash on Editor shutdown
+	if (FSlateApplication::IsInitialized())
 	{
 		FSlateApplication::Get().UnregisterInputPreProcessor(InputProcessor);
 	}
