[AI Assistant]

FExecuteWhenReady -
- Removed method to return array of deferred execution functions, since they could be potentially altered, or called. Instead added method (for testing) that just returns how many are stored.
- Added locks around accessors/modifcations of deferred execution functions.
- Move can't copy the locks, so added special move functions, that also avoid a potential hard-to-debug deadlock.

AIAssistantWebBrowser -
OnLoadCompleted(), properly advance ready state of AI Assistant web browser - in ALL cases.

AIAssistant -
Check for potential null dereference of shared pointer in ShowContextMenu().
This shouldn't be possible, but it's a good defensive fix for a theoretical edge case.

diff --git a/Engine/Plugins/Experimental/AIAssistant/Source/AIAssistant/Private/AIAssistant.cpp b/Engine/Plugins/Experimental/AIAssistant/Source/AIAssistant/Private/AIAssistant.cpp
--- a/Engine/Plugins/Experimental/AIAssistant/Source/AIAssistant/Private/AIAssistant.cpp
+++ b/Engine/Plugins/Experimental/AIAssistant/Source/AIAssistant/Private/AIAssistant.cpp
@@ -57,7 +57,7 @@ static const FName AIAssistantTabName("AIAssistant");
 	FGlobalTabmanager::Get()->RegisterNomadTabSpawner(AIAssistantTabName, FOnSpawnTab::CreateRaw(this, &FAIAssistantModule::OnSpawnPluginTab))
 		.SetDisplayName(LOCTEXT("FAIAssistantTabTitle", "AI Assistant"))
 		.SetMenuType(ETabSpawnerMenuType::Hidden)
-		.SetIcon(FSlateIcon("AIAssistantStyle", "AIAssistant.OpenPluginWindow"));
+		.SetIcon(FSlateIcon("AIAssistantStyle","AIAssistant.OpenPluginWindow"));
 
 	
 	InputProcessor = MakeShared<FAIAssistantInputProcessor>(PluginCommands);
@@ -113,6 +113,11 @@ TSharedPtr<SAIAssistantWebBrowser> FAIAssistantModule::GetAIAssistantWebBrowserW
 
 void FAIAssistantModule::ShowContextMenu(const FString& SelectedString, const FVector2f& ClientLocation) const
 {
+	if (!AIAssistantWebBrowserWidget.IsValid())
+	{
+		return;
+	}
+	
 	const bool bIsSelectedTextValid = (!SelectedString.IsEmpty());
 
 	
@@ -187,7 +192,7 @@ TSharedRef<SDockTab> FAIAssistantModule::OnSpawnPluginTab(const FSpawnTabArgs& S
 		.TabRole(ETabRole::NomadTab)
 		.OnTabClosed_Lambda([this](TSharedRef<SDockTab>) -> void
 		{
-			AIAssistantWebBrowserWidget->OnClosed();
+			AIAssistantWebBrowserWidget->OnClosed();			
 		})
 		[
 			AIAssistantWebBrowserWidget.ToSharedRef()
diff --git a/Engine/Plugins/Experimental/AIAssistant/Source/AIAssistant/Private/AIAssistantExecuteWhenReady.cpp b/Engine/Plugins/Experimental/AIAssistant/Source/AIAssistant/Private/AIAssistantExecuteWhenReady.cpp
--- a/Engine/Plugins/Experimental/AIAssistant/Source/AIAssistant/Private/AIAssistantExecuteWhenReady.cpp
+++ b/Engine/Plugins/Experimental/AIAssistant/Source/AIAssistant/Private/AIAssistantExecuteWhenReady.cpp
@@ -3,6 +3,23 @@
 
 #include "AIAssistantExecuteWhenReady.h"
 
+#include "Misc/ScopeRWLock.h"
+#include <functional>
+
+
+UE::AIAssistant::FExecuteWhenReady::FExecuteWhenReady(FExecuteWhenReady&& Other) noexcept
+{
+	MoveFrom(Other);
+}
+
+
+UE::AIAssistant::FExecuteWhenReady& UE::AIAssistant::FExecuteWhenReady::operator=(FExecuteWhenReady&& RHS) noexcept
+{
+	MoveFrom(RHS);
+	
+	return *this;
+}
+
 
 void UE::AIAssistant::FExecuteWhenReady::Enqueue(FDeferredExecutionFunction&& DeferredExecutionFunction)
 {
@@ -12,6 +29,8 @@ void UE::AIAssistant::FExecuteWhenReady::Enqueue(FDeferredExecutionFunction&& De
 	}
 	else 
 	{
+		FWriteScopeLock WriteScopeLock(DeferredExecutionFunctionRWLock);
+		
 		DeferredExecutionFunctions.Add(MoveTemp(DeferredExecutionFunction));
 	}
 }
@@ -21,6 +40,8 @@ void UE::AIAssistant::FExecuteWhenReady::Update()
 {
 	if (GetReadyState() == EReadyState::Execute)
 	{
+		FWriteScopeLock WriteScopeLock(DeferredExecutionFunctionRWLock);
+
 		for (TArray<FDeferredExecutionFunction> Functions = MoveTemp(DeferredExecutionFunctions);
 			FDeferredExecutionFunction& Function : Functions)
 		{
@@ -36,11 +57,41 @@ void UE::AIAssistant::FExecuteWhenReady::Update()
 
 void UE::AIAssistant::FExecuteWhenReady::Reset()
 {
+	FWriteScopeLock WriteScopeLock(DeferredExecutionFunctionRWLock);
+
 	DeferredExecutionFunctions.Empty();
 }
 
 
-TArray<UE::AIAssistant::FExecuteWhenReady::FDeferredExecutionFunction>& UE::AIAssistant::FExecuteWhenReady::GetDeferredExecutionFunctions()
+int32 UE::AIAssistant::FExecuteWhenReady::GetNumDeferredExecutionFunctions() const
 {
-	return DeferredExecutionFunctions;
+	FReadScopeLock ReadScopedLock(DeferredExecutionFunctionRWLock);
+
+	return DeferredExecutionFunctions.Num();
+}
+
+
+void UE::AIAssistant::FExecuteWhenReady::MoveFrom(FExecuteWhenReady& Other) noexcept
+{
+	if (this == &Other)
+	{
+		return;
+	}
+
+	
+	// Looks strange, but we're being extra careful to avoid a particular kind of deadlock, which would be SUPER HARD to debug. Use pointer values to
+	// find a consistent locking order - in the very unlikely case that one thread tries to move two of these classes, while another thread tries to
+	// move the same two classes at the same time but in the opposite order. This way, they can't both acquire a first lock that is the second lock
+	// of the other one, with both would ending up getting blocked trying to acquire their second lock. To enforce the lock order, make the lower
+	// pointer value always lock first, and the higher pointer value always lock second. (It's technically safest to use std::less for the
+	// comparison.)
+	
+	const bool bLockThisBeforeOther = std::less<const FExecuteWhenReady*>()(this, &Other); // ..see comment
+	FWriteScopeLock WriteScopeLock1((bLockThisBeforeOther ? this : &Other)->DeferredExecutionFunctionRWLock);
+	FWriteScopeLock WriteScopeLock2((bLockThisBeforeOther ? &Other : this)->DeferredExecutionFunctionRWLock);
+
+
+	// We want to move everything, but are not allowed to move the lock.
+	
+	DeferredExecutionFunctions = MoveTemp(Other.DeferredExecutionFunctions);
 }
diff --git a/Engine/Plugins/Experimental/AIAssistant/Source/AIAssistant/Private/AIAssistantExecuteWhenReady.h b/Engine/Plugins/Experimental/AIAssistant/Source/AIAssistant/Private/AIAssistantExecuteWhenReady.h
--- a/Engine/Plugins/Experimental/AIAssistant/Source/AIAssistant/Private/AIAssistantExecuteWhenReady.h
+++ b/Engine/Plugins/Experimental/AIAssistant/Source/AIAssistant/Private/AIAssistantExecuteWhenReady.h
@@ -6,6 +6,7 @@
 
 #include "Templates/Function.h"
 #include "Containers/Array.h"
+#include "HAL/CriticalSection.h"
 
 
 //
@@ -43,8 +44,8 @@ namespace UE::AIAssistant
 		FExecuteWhenReady& operator=(const FExecuteWhenReady&) = delete;
 
 		// Allow moves.
-		FExecuteWhenReady(FExecuteWhenReady&&) noexcept = default;
-		FExecuteWhenReady& operator=(FExecuteWhenReady&&) noexcept = default;
+		FExecuteWhenReady(FExecuteWhenReady&&) noexcept;
+		FExecuteWhenReady& operator=(FExecuteWhenReady&&) noexcept;
 
 
 		/**
@@ -92,19 +93,26 @@ namespace UE::AIAssistant
 
 
 		/**
-		 * Accessor for internal deferred execution functions. Will be empty after execution or rejection.
-		 * @return Current array of deferred execution functions. 
+		 * How many deferred execution functions are stored? Will be 0 after execution or rejection.
+		 * @return Current number of deferred execution functions stored. 
 		 */
-		TArray<FDeferredExecutionFunction>& GetDeferredExecutionFunctions();
+
+		int32 GetNumDeferredExecutionFunctions() const;
 		
 
 	private:
 		
 
+		// Used during access/modification of deferred execution functions.
+		mutable FRWLock DeferredExecutionFunctionRWLock;
+
 		// Either executes incoming deferred execution function if ready now, or saves it to execute when ready later.
 		void Enqueue(FDeferredExecutionFunction&& DeferredExecutionFunction);
 
-		// Asynchronous functions to execute in order, when we become ready to execute them.
+		// Deferred functions to execute in order, when we become ready to execute them.
 		TArray<FDeferredExecutionFunction> DeferredExecutionFunctions;
+
+		// Used for custom moves.
+		void MoveFrom(FExecuteWhenReady& Other) noexcept;
 	};
 }
diff --git a/Engine/Plugins/Experimental/AIAssistant/Source/AIAssistant/Private/AIAssistantWebBrowser.cpp b/Engine/Plugins/Experimental/AIAssistant/Source/AIAssistant/Private/AIAssistantWebBrowser.cpp
--- a/Engine/Plugins/Experimental/AIAssistant/Source/AIAssistant/Private/AIAssistantWebBrowser.cpp
+++ b/Engine/Plugins/Experimental/AIAssistant/Source/AIAssistant/Private/AIAssistantWebBrowser.cpp
@@ -162,6 +162,10 @@ void SAIAssistantWebBrowser::Construct(const FArguments& InArgs)
 						}
 					});
 			}
+			else
+			{
+				UpdateWebBrowserLoadState(EWebBrowserLoadState::LoadComplete);
+			}
 		});
 	
 	
diff --git a/Engine/Plugins/Experimental/AIAssistant/Source/AIAssistant/Private/Tests/AIAssistantExecuteWhenReadyTest.cpp b/Engine/Plugins/Experimental/AIAssistant/Source/AIAssistant/Private/Tests/AIAssistantExecuteWhenReadyTest.cpp
--- a/Engine/Plugins/Experimental/AIAssistant/Source/AIAssistant/Private/Tests/AIAssistantExecuteWhenReadyTest.cpp
+++ b/Engine/Plugins/Experimental/AIAssistant/Source/AIAssistant/Private/Tests/AIAssistantExecuteWhenReadyTest.cpp
@@ -37,17 +37,17 @@ void GeneralFakeExecuteWhenReadyTestFunction(FAutomationTestBase& AutomationTest
 		FakeExecuteWhenReady.SetFakeStateCount(I);
 		FakeExecuteWhenReady.Update();
 
-		if (TestMode == UE::AIAssistant::FFakeExecuteWhenReady::ExecuteWhenStateHitsValue)
+		if (TestMode == FFakeExecuteWhenReady::ExecuteWhenStateHitsValue)
 		{
 			if (I >= FFakeExecuteWhenReady::FakeStateTransitionCount)
 			{
 				AutomationTest.TestEqual("TestString", TestStringCurrent, TestStringWhenExecuted);
-				AutomationTest.TestEqual("DeferredExecutionFunctionCount", FakeExecuteWhenReady.GetDeferredExecutionFunctions().Num(), 0);
+				AutomationTest.TestEqual("DeferredExecutionFunctionCount", FakeExecuteWhenReady.GetNumDeferredExecutionFunctions(), 0);
 			}
 			else 
 			{
 				AutomationTest.TestEqual("TestString", TestStringCurrent, TestStringWhenWaiting);
-				AutomationTest.TestEqual("DeferredExecutionFunctionCount", FakeExecuteWhenReady.GetDeferredExecutionFunctions().Num(), 2);
+				AutomationTest.TestEqual("DeferredExecutionFunctionCount", FakeExecuteWhenReady.GetNumDeferredExecutionFunctions(), 2);
 			}
 		}
 		else if (TestMode == FFakeExecuteWhenReady::RejectWhenStateHitsValue)
@@ -55,18 +55,18 @@ void GeneralFakeExecuteWhenReadyTestFunction(FAutomationTestBase& AutomationTest
 			if (I >= FFakeExecuteWhenReady::FakeStateTransitionCount)
 			{
 				AutomationTest.TestEqual("TestString", TestStringCurrent, TestStringWhenWaiting);
-				AutomationTest.TestEqual("DeferredExecutionFunctionCount", FakeExecuteWhenReady.GetDeferredExecutionFunctions().Num(), 0);
+				AutomationTest.TestEqual("DeferredExecutionFunctionCount", FakeExecuteWhenReady.GetNumDeferredExecutionFunctions(), 0);
 			}
 			else 
 			{
 				AutomationTest.TestEqual("TestString", TestStringCurrent, TestStringWhenWaiting);
-				AutomationTest.TestEqual("DeferredExecutionFunctionCount", FakeExecuteWhenReady.GetDeferredExecutionFunctions().Num(), 2);
+				AutomationTest.TestEqual("DeferredExecutionFunctionCount", FakeExecuteWhenReady.GetNumDeferredExecutionFunctions(), 2);
 			}
 		}
 		else if (TestMode == FFakeExecuteWhenReady::IgnoreWhenStateHitsValue)
 		{
 			AutomationTest.TestEqual("TestString", TestStringCurrent, TestStringWhenWaiting);
-			AutomationTest.TestEqual("DeferredExecutionFunctionCount", FakeExecuteWhenReady.GetDeferredExecutionFunctions().Num(), 2);
+			AutomationTest.TestEqual("DeferredExecutionFunctionCount", FakeExecuteWhenReady.GetNumDeferredExecutionFunctions(), 2);
 		}
 	}
 }
