[AI Assistant]

Doc change.

diff --git a/Engine/Plugins/Experimental/AIAssistant/Source/AIAssistant/Private/AIAssistantExecuteWhenReady.h b/Engine/Plugins/Experimental/AIAssistant/Source/AIAssistant/Private/AIAssistantExecuteWhenReady.h
--- a/Engine/Plugins/Experimental/AIAssistant/Source/AIAssistant/Private/AIAssistantExecuteWhenReady.h
+++ b/Engine/Plugins/Experimental/AIAssistant/Source/AIAssistant/Private/AIAssistantExecuteWhenReady.h
@@ -18,14 +18,13 @@
 //
 // A class that inherits from this class -
 //		- can have callables added to it via ExecuteWhenReady() , which will be executed, in order, once the class says its ready
-//		- should itself override GetReadyState() to notify its ready state
-//		- should itself call Update() when necessary to trigger possible execution of the callables provided according to its ready state
+//		- should itself override GetExecuteWhenReadyState() to notify its ready state
+//		- should itself call UpdateExecuteWhenReady() when necessary to execute the callables provided, according to its ready state
 //
 // Example -
-// A class that uses a web browser may need to wait until a web page is loaded before executing certain functions, such as sending da ""
-// ta to that web
+// A class that uses a web browser may need to wait until a web page is loaded before executing certain functions, such as sending data to that web
 // page. If that class inherits from this class, those functions can be added as callables to it. That class can monitor whether the web page is
-// loaded, and change the state returned by GetReadyState() when it is. A call to Update() will then execute the callables that were added to it. 
+// loaded. Whenever web pages are done loading, it would set it's ready state and then can call the update function.  
 //
 
 
