[AI Assistant] Fix possible invalid reference in slate querier

diff --git a/Engine/Plugins/Experimental/AIAssistant/Source/AIAssistant/Private/AIAssistantSlateQuerier.cpp b/Engine/Plugins/Experimental/AIAssistant/Source/AIAssistant/Private/AIAssistantSlateQuerier.cpp
--- a/Engine/Plugins/Experimental/AIAssistant/Source/AIAssistant/Private/AIAssistantSlateQuerier.cpp
+++ b/Engine/Plugins/Experimental/AIAssistant/Source/AIAssistant/Private/AIAssistantSlateQuerier.cpp
@@ -418,8 +418,11 @@ static FText FindItemName(const FWidgetPath& WidgetPath, FAIAssistantSlateQueryC
 	if (ConsoleInputBox->GetType() != "SNullWidgetContent")
 	{
 		TSharedRef<SWidget> EditableText = FindChildWidgetOfType(ConsoleInputBox, "SMultiLineEditableTextBox");
-		TSharedRef<SMultiLineEditableTextBox> ThisMultiLineEditableTextBoxCast = StaticCastSharedRef<SMultiLineEditableTextBox>(EditableText);
-		HintText = ThisMultiLineEditableTextBoxCast->GetHintText();
+		if (EditableText->GetType() != "SNullWidgetContent")
+		{
+			TSharedRef<SMultiLineEditableTextBox> ThisMultiLineEditableTextBoxCast = StaticCastSharedRef<SMultiLineEditableTextBox>(EditableText);
+			HintText = ThisMultiLineEditableTextBoxCast->GetHintText();
+		}
 		OutName = FText::Format(LOCTEXT("ConsoleInputBoxFormat", "{0}"), (HintText.IsEmpty() ? LOCTEXT("ConsoleInputBoxDefaultName","Console") : HintText));
 		OutDescriptor = LOCTEXT("ItemDescriptor_ConsoleInputBox", " input box");
 		SlateQueryContext.LastPickedWidget = ConsoleInputBox.ToSharedPtr();
