[AI Assistant] Remove invalid "pragma once" from source file

diff --git a/Engine/Plugins/Experimental/AIAssistant/Source/AIAssistant/Private/AIAssistantConsole.cpp b/Engine/Plugins/Experimental/AIAssistant/Source/AIAssistant/Private/AIAssistantConsole.cpp
--- a/Engine/Plugins/Experimental/AIAssistant/Source/AIAssistant/Private/AIAssistantConsole.cpp
+++ b/Engine/Plugins/Experimental/AIAssistant/Source/AIAssistant/Private/AIAssistantConsole.cpp
@@ -1,6 +1,4 @@
 // Copyright Epic Games, Inc. All Rights Reserved.
-#pragma once
-
 #include "AIAssistantConsole.h"
 
 #include "Async/UniqueLock.h"
