[AIAssistant]

Minor documentation fix.

diff --git a/Engine/Plugins/Experimental/AIAssistant/Source/AIAssistant/Private/AIAssistantWebBrowser.h b/Engine/Plugins/Experimental/AIAssistant/Source/AIAssistant/Private/AIAssistantWebBrowser.h
--- a/Engine/Plugins/Experimental/AIAssistant/Source/AIAssistant/Private/AIAssistantWebBrowser.h
+++ b/Engine/Plugins/Experimental/AIAssistant/Source/AIAssistant/Private/AIAssistantWebBrowser.h
@@ -22,8 +22,8 @@
 // A UObject can be bound to SWebBrowser.
 // That will make that Object's UFUNCTIONs available for calling in JavaScript.
 // They will be available as "window.ue.namespace.functionname".
-// In our case, we bind the UAIAssistantSubsystem to this web browser, as 'aiassistantsubsystem'.
-// So, for example, JavaScript can call - "window.ue.aiassistantsubsystem.executepythonscriptviajavascript(code)".
+// In our case, we may choose to bind the UAIAssistantSubsystem to this web browser, as 'aiassistantsubsystem'.
+// So, for example, JavaScript would then be able to call - "window.ue.aiassistantsubsystem.executepythonscriptviajavascript(code)".
 // IMPORTANT - Yes, it must be all lowercase, or the call will fail.
 //
 
@@ -32,7 +32,7 @@ class SAIAssistantWebBrowser :
 	public SCompoundWidget,
 	public UE::AIAssistant::IWebJavaScriptDelegateBinder,
 	public UE::AIAssistant::IWebJavaScriptExecutor,
-	public UE::AIAssistant::FExecuteWhenReady // ..executes any added callables to execute, once the proper web page is loaded
+	public UE::AIAssistant::FExecuteWhenReady
 {
 public:
 
