[AI Assistant] Fix LOCTEXT entries for assistant query and context
Provide named parameters where made sense
Remove concatenation, duplicated LOCTEXT keys.

diff --git a/Engine/Plugins/Experimental/AIAssistant/Source/AIAssistant/Private/AIAssistantSlateQuerier.cpp b/Engine/Plugins/Experimental/AIAssistant/Source/AIAssistant/Private/AIAssistantSlateQuerier.cpp
--- a/Engine/Plugins/Experimental/AIAssistant/Source/AIAssistant/Private/AIAssistantSlateQuerier.cpp
+++ b/Engine/Plugins/Experimental/AIAssistant/Source/AIAssistant/Private/AIAssistantSlateQuerier.cpp
@@ -254,7 +254,10 @@ static FText GenerateModeToolContext(const FWidgetPath& WidgetPath)
 						FText ToolDisplayName = ModeToolkit.Pin()->GetActiveToolDisplayName();
 						if (!ToolDisplayName.IsEmpty())
 						{
-							return FText::Format(LOCTEXT("EditorName_ModeToolFormat", "The {0} tool is active for {1} mode."), ToolDisplayName, ModeInfo.Name);
+							FFormatNamedArguments Args;
+							Args.Add(TEXT("ToolName"), ToolDisplayName);
+							Args.Add(TEXT("ModeName"), ModeInfo.Name);
+							return FText::Format(LOCTEXT("EditorName_ModeToolFormat", "The {ToolName} tool is active for {ModeName} mode."), Args);
 						}
 					}
 				}
@@ -408,8 +411,8 @@ static FText FindItemName(const FWidgetPath& WidgetPath, FAIAssistantSlateQueryC
 	}
 	if (SearchBox->GetType() != "SNullWidgetContent")
 	{
-		OutName = FText::Format(LOCTEXT("SearchBoxFormat", "{0}"), (HintText.IsEmpty()) ? LOCTEXT("SearchBoxDefaultName","search") : HintText);
-		OutDescriptor = LOCTEXT("ItemDescriptor_SearchBox", " box");
+		OutName = FText::Format(LOCTEXT("SearchBoxFormat", "{0}"), (HintText.IsEmpty()) ? LOCTEXT("SearchBoxDefaultName", "default") : HintText);
+		OutDescriptor = LOCTEXT("ItemDescriptor_SearchBox", "search box");
 		SlateQueryContext.LastPickedWidget = SearchBox.ToSharedPtr();
 	}
 
@@ -424,7 +427,7 @@ static FText FindItemName(const FWidgetPath& WidgetPath, FAIAssistantSlateQueryC
 			HintText = ThisMultiLineEditableTextBoxCast->GetHintText();
 		}
 		OutName = FText::Format(LOCTEXT("ConsoleInputBoxFormat", "{0}"), (HintText.IsEmpty() ? LOCTEXT("ConsoleInputBoxDefaultName","Console") : HintText));
-		OutDescriptor = LOCTEXT("ItemDescriptor_ConsoleInputBox", " input box");
+		OutDescriptor = LOCTEXT("ItemDescriptor_ConsoleInputBox", "input box");
 		SlateQueryContext.LastPickedWidget = ConsoleInputBox.ToSharedPtr();
 	}
 
@@ -502,9 +505,12 @@ static FText FindItemName(const FWidgetPath& WidgetPath, FAIAssistantSlateQueryC
 
 	if (!OutName.IsEmpty())
 	{
-		return FText::Format(LOCTEXT("ItemName_Format", " the \"{0}\" {1}"), { OutName, OutDescriptor });
+		FFormatNamedArguments Args;
+		Args.Add(TEXT("ItemName"), OutName);
+		Args.Add(TEXT("ItemDescriptor"), OutDescriptor);
+		return FText::Format(LOCTEXT("ItemName_Format", " the \"{ItemName}\" {ItemDescriptor}"), Args);
 	}
-	
+
 	return OutName;
 }
 
@@ -535,7 +541,9 @@ static FText FindTabName(const FWidgetPath& WidgetPath, FAIAssistantSlateQueryCo
 	
 	if (!OutName.IsEmpty())
 	{
-		return FText::Format(LOCTEXT("TabName_Format", " the {0} panel"), { OutName });
+		FFormatNamedArguments Args;
+		Args.Add(TEXT("PanelName"), OutName);
+		return FText::Format(LOCTEXT("TabName_Format", " the {PanelName} panel"), Args);
 	}
 	
 	return OutName;
@@ -551,19 +559,23 @@ static FText GenerateDetailsViewContext(const FWidgetPath& WidgetPath, FAIAssist
 	if (DetailsView->GetType() != "SNullWidgetContent")
 	{
 		const TSharedRef<IDetailsView> DetailsViewCast = StaticCastSharedRef<IDetailsView>(DetailsView);
+		FText QueryContextDetailsObjectAddition =
+			LOCTEXT("QueryContextDetailsObjectAddition",
+					"The property panel I'm working in is editing at least one actor of the class \"{ClassName}\".");
+		FFormatNamedArguments ClassNameArgs;
 		if (DetailsViewCast->GetSelectedActors().Num() > 0)
 		{
 			FString ActorName;
 			DetailsViewCast->GetSelectedActors()[0]->GetClass()->GetName(ActorName);
-			const FText ActorClass = FText::FromString(ActorName);
-			SlateQueryContext.GeneratedContextItems.Add(FText::Format(LOCTEXT("QueryContextDetailsObjectAddition", "The property panel I'm working in is editing at least one actor of the class \"{0}\"."), ActorClass));
+			ClassNameArgs.Add(TEXT("ClassName"), FText::FromString(ActorName));
+			SlateQueryContext.GeneratedContextItems.Add(FText::Format(QueryContextDetailsObjectAddition, ClassNameArgs));
 		}
 		else if (DetailsViewCast->GetSelectedObjects().Num() > 0)
 		{
 			FString ObjectName;
 			DetailsViewCast->GetSelectedObjects()[0]->GetClass()->GetName(ObjectName);
-			const FText ObjectClass = FText::FromString(ObjectName);
-			SlateQueryContext.GeneratedContextItems.Add(FText::Format(LOCTEXT("QueryContextDetailsObjectAddition", "The property panel I'm working in is editing at least one object of the class \"{0}\"."), ObjectClass));
+			ClassNameArgs.Add(TEXT("ClassName"), FText::FromString(ObjectName));
+			SlateQueryContext.GeneratedContextItems.Add(FText::Format(QueryContextDetailsObjectAddition, ClassNameArgs));
 		}
 	}
 
@@ -613,12 +625,23 @@ static FText FindEditorName(const FWidgetPath& WidgetPath, FAIAssistantSlateQuer
 					{
 						if (ModeString.IsEmpty() && ModeManager->IsModeActive(ModeInfo.ID) && !BlockedModeIDs.Contains(ModeInfo.ID))
 						{
-							ModeString = FText::Format(LOCTEXT("EditorName_ModeFormat", " with {0} mode active"), { ModeInfo.Name });
+							FFormatNamedArguments Args;
+							Args.Add(TEXT("ModeName"), ModeInfo.Name);
+							ModeString = FText::Format(LOCTEXT("EditorName_ModeFormat", "with {ModeName} mode active"), Args);
 						}
 					}
 				}
 				SlateQueryContext.LastPickedWidget = LevelEditor.ToSharedPtr();
-				return FText::Format(LOCTEXT("EditorName_LevelEditorAndMode", " the Level Editor{0}"), ModeString);
+				if (ModeString.IsEmpty())
+				{
+					return LOCTEXT("EditorName_LevelEditor", " the Level Editor");
+				}
+				else
+				{
+					FFormatNamedArguments Args;
+					Args.Add(TEXT("WithModeActive"), ModeString);
+					return FText::Format(LOCTEXT("EditorName_LevelEditorAndMode", " the Level Editor {WithModeActive}"), Args);
+				}
 			}
 
 			// if not, are we in a different asset editor?
@@ -641,7 +664,9 @@ static FText FindEditorName(const FWidgetPath& WidgetPath, FAIAssistantSlateQuer
 							if (ThisEditor->GetAssociatedTabManager() == PickedWidgetTabManager)
 							{
 								SlateQueryContext.LastPickedWidget = TabStack.ToSharedPtr();
-								return FText::Format(LOCTEXT("EditorName_AssetEditorFormat", " the {0} editor"), { FText::FromName(ThisEditor->GetEditorName()) });
+								FFormatNamedArguments Args;
+								Args.Add(TEXT("AssetEditorName"), FText::FromName(ThisEditor->GetEditorName()));
+								return FText::Format(LOCTEXT("EditorName_AssetEditorFormat", " the {AssetEditorName} editor"), Args);
 							}
 						}
 					}
@@ -674,7 +699,9 @@ static FText FindEditorName(const FWidgetPath& WidgetPath, FAIAssistantSlateQuer
 		if (!DrawerName.IsEmpty())
 		{
 			SlateQueryContext.LastPickedWidget = DrawerWidget.ToSharedPtr();
-			return FText::Format(LOCTEXT("DrawerNameFormat", "the {0} drawer"), DrawerName);
+			FFormatNamedArguments Args;
+			Args.Add(TEXT("DrawerName"), DrawerName);
+			return FText::Format(LOCTEXT("DrawerNameFormat", "the {DrawerName} drawer"), Args);
 		}
 	}
 
@@ -694,7 +721,9 @@ static FText FindEditorName(const FWidgetPath& WidgetPath, FAIAssistantSlateQuer
 	if (!WindowTitle.IsEmpty())
 	{
 		SlateQueryContext.LastPickedWidget = WidgetWindow.ToSharedPtr();
-		return FText::Format(LOCTEXT("WindowNameFormat_Title", " the {0} window"), WindowTitle);
+		FFormatNamedArguments Args;
+		Args.Add(TEXT("WindowTitle"), WindowTitle);
+		return FText::Format(LOCTEXT("WindowNameFormat_Title", " the {WindowTitle} window"), Args);
 	}
 
 	// TODO: how to find other editor contexts?
@@ -816,13 +845,16 @@ void UE::AIAssistant::SlateQuerier::QueryAIAssistantAboutSlateWidgetUnderCursor(
 			const FText NameToQuery = (!ItemName.IsEmpty()) ? ItemName : (!TabName.IsEmpty()) ? TabName : WindowName;
 
 			static const FText QueryFormat = LOCTEXT("QueryFormat1",
-				"I would like to know what {0} does in the Unreal Editor UI.");
-			SlateQueryContext.GeneratedQuery = FText::Format(QueryFormat, NameToQuery);
+				"I would like to know what {NameToQuery} does in the Unreal Editor UI.");
+			FFormatNamedArguments QueryArgs;
+			QueryArgs.Add(TEXT("NameToQuery"), NameToQuery);
+			SlateQueryContext.GeneratedQuery = FText::Format(QueryFormat, QueryArgs);
 			
 			static const FText QueryInstructionsFormat = LOCTEXT("QueryInstructionsFormat1",
-				"Provide a short, conversational answer that describes what the widget does. {0}");
-			SlateQueryContext.GeneratedQueryInstructions =
-				FText::Format(QueryInstructionsFormat, QueryAdditionalInstructions);
+				"Provide a short, conversational answer that describes what the widget does. {QueryAdditionalInstructions}");
+			FFormatNamedArguments InstructionArgs;
+			InstructionArgs.Add(TEXT("QueryAdditionalInstructions"), QueryAdditionalInstructions);
+			SlateQueryContext.GeneratedQueryInstructions = FText::Format(QueryInstructionsFormat, InstructionArgs);
 		}
 		
 		if (ItemName.IsEmpty())
@@ -830,13 +862,16 @@ void UE::AIAssistant::SlateQuerier::QueryAIAssistantAboutSlateWidgetUnderCursor(
 			if (!TextUnderCursor.IsEmpty())
 			{
 				static const FText QueryFormat = LOCTEXT("QueryFormat2",
-					"What does \"{0}\" mean in the Unreal Editor UI?");
-				SlateQueryContext.GeneratedQuery = FText::Format(QueryFormat, TextUnderCursor);
+					"What does \"{ObjectText}\" mean in the Unreal Editor UI?");
+				FFormatNamedArguments Args;
+				Args.Add(TEXT("ObjectText"), TextUnderCursor);
+				SlateQueryContext.GeneratedQuery = FText::Format(QueryFormat, Args);
 				
 				static const FText QueryInstructionsFormat = LOCTEXT("QueryInstructionsFormat2",
-					"Provide a short, conversational answer that describes what I should know about it. {0}");
-				SlateQueryContext.GeneratedQueryInstructions =
-					FText::Format(QueryInstructionsFormat, QueryAdditionalInstructions);
+					"Provide a short, conversational answer that describes what I should know about it. {QueryAdditionalInstructions}");
+				FFormatNamedArguments InstructionArgs;
+				InstructionArgs.Add(TEXT("QueryAdditionalInstructions"), QueryAdditionalInstructions);
+				SlateQueryContext.GeneratedQueryInstructions = FText::Format(QueryInstructionsFormat, InstructionArgs);
 			}
 		}
 
@@ -852,9 +887,19 @@ void UE::AIAssistant::SlateQuerier::QueryAIAssistantAboutSlateWidgetUnderCursor(
 	{
 		if (!TabName.IsEmpty() || !WindowName.IsEmpty())
 		{
-			const FText TabConnector = (!TabName.IsEmpty() && !WindowName.IsEmpty()) ? LOCTEXT("Query_TabConnector", " of") : FText();
-
-			SlateQueryContext.GeneratedContextItems.Add(FText::Format(LOCTEXT("QueryFormat", "I am working in {0}{1}{2}."), TabName, TabConnector, WindowName));
+			FFormatNamedArguments Args;
+			if (!TabName.IsEmpty() && !WindowName.IsEmpty())
+			{
+				Args.Add(TEXT("TabName"), TabName);
+				Args.Add(TEXT("WindowName"), WindowName);
+				SlateQueryContext.GeneratedContextItems.Add(FText::Format(LOCTEXT("QueryContextTabWindowLocation", "I am working in {TabName} of {WindowName}."), Args));
+			}
+			else
+			{
+				Args.Add(TEXT("MaybeEmptyTabName"), TabName);
+				Args.Add(TEXT("MaybeEmptyWindowName"), WindowName);
+				SlateQueryContext.GeneratedContextItems.Add(FText::Format(LOCTEXT("QueryContextTabOrWindowLocation", "I am working in {MaybeEmptyTabName}{MaybeEmptyWindowName}."), Args));
+			}
 
 			if (const FText ModeToolContext = GenerateModeToolContext(WidgetPath);
 				!ModeToolContext.IsEmpty())
@@ -871,15 +916,21 @@ void UE::AIAssistant::SlateQuerier::QueryAIAssistantAboutSlateWidgetUnderCursor(
 
 		if (!ToolTipText.IsEmpty())
 		{
-			SlateQueryContext.GeneratedContextItems.Add(FText::Format(LOCTEXT("QueryContextToolTipAddition", "The text of the current tooltip is: \"{0}\""), ToolTipText));
+			FText ToolTipContextFormat = LOCTEXT("QueryContextToolTipAddition", "The text of the current tooltip is: \"{ToolTipText}\"");
+			FFormatNamedArguments Args;
+			Args.Add(TEXT("ToolTipText"), ToolTipText);
+			FText ToolTipContext = FText::Format(ToolTipContextFormat, Args);
+			SlateQueryContext.GeneratedContextItems.Add(ToolTipContext);
 		}
 
-		
 		if (!SlateQueryContext.GeneratedContextItems.IsEmpty())
 		{
 			for (const FText& ContextItem : SlateQueryContext.GeneratedContextItems)
 			{
-				SlateQueryContext.GeneratedContext = FText::Format(LOCTEXT("LastGeneratedContentFormat", "{0} {1}"), SlateQueryContext.GeneratedContext, ContextItem);
+				FFormatNamedArguments Args;
+				Args.Add(TEXT("GeneratedContextSoFar"), SlateQueryContext.GeneratedContext);
+				Args.Add(TEXT("ContextItemToAdd"), ContextItem);
+				SlateQueryContext.GeneratedContext = FText::Format(LOCTEXT("LastGeneratedContentFormat", "{GeneratedContextSoFar} {ContextItemToAdd}"), Args);
 			}
 		}
 	}
@@ -889,7 +940,6 @@ void UE::AIAssistant::SlateQuerier::QueryAIAssistantAboutSlateWidgetUnderCursor(
 	SlateQueryContext.LastPickedWidget.Reset();
 	SlateQueryContext.GeneratedContextItems.Empty();
 
-
 	// Send widget query to AI Assistant.
 
 	AIAssistantWebBrowserWidget->ExecuteWhenReady([SlateQueryContext]
