[AIAssistant]

Alt+F1 opens AI Assistant.

diff --git a/Engine/Plugins/Experimental/AIAssistant/Source/AIAssistant/Private/AIAssistantCommands.cpp b/Engine/Plugins/Experimental/AIAssistant/Source/AIAssistant/Private/AIAssistantCommands.cpp
--- a/Engine/Plugins/Experimental/AIAssistant/Source/AIAssistant/Private/AIAssistantCommands.cpp
+++ b/Engine/Plugins/Experimental/AIAssistant/Source/AIAssistant/Private/AIAssistantCommands.cpp
@@ -17,7 +17,7 @@ FAIAssistantCommands::FAIAssistantCommands() :
 
 /*virtual*/ void FAIAssistantCommands::RegisterCommands()
 {
-	UI_COMMAND(OpenAIAssistantTab, "AI Assistant", "Bring up AI Assistant window", EUserInterfaceActionType::Button, FInputChord());
+	UI_COMMAND(OpenAIAssistantTab, "AI Assistant", "Bring up AI Assistant window", EUserInterfaceActionType::Button, FInputChord(EModifierKey::Alt, EKeys::F1));
 
 	UI_COMMAND(SlateQueryCommand, "Query UI with AI Assistant", "Query the UI widget under the mouse with AI Assistant", EUserInterfaceActionType::Button, FInputChord(EKeys::F1));
 }
