[AI Assistant] tweaks to UI query prompt and widget/object selection
- differentiate prompt & context for objects vs widget items
- better display identification of assets in content browser
- better display identification of icon-only buttons
- improve query context for objects in Outliner
- add style suggestions to prompt to get more consistently structured responses

diff --git a/Engine/Plugins/Experimental/AIAssistant/Source/AIAssistant/Private/AIAssistantSlateQuerier.cpp b/Engine/Plugins/Experimental/AIAssistant/Source/AIAssistant/Private/AIAssistantSlateQuerier.cpp
--- a/Engine/Plugins/Experimental/AIAssistant/Source/AIAssistant/Private/AIAssistantSlateQuerier.cpp
+++ b/Engine/Plugins/Experimental/AIAssistant/Source/AIAssistant/Private/AIAssistantSlateQuerier.cpp
@@ -40,6 +40,10 @@ struct FAIAssistantSlateQueryContext
 {
 	TSharedPtr<SWidget> LastPickedWidget;
 	TArray<FText> GeneratedContextItems;
+	FText CurrentToolTipText = FText();
+	bool bIsUIWidget = false;
+	bool bIsObject = false;
+	bool bInOutliner = false;
 
 	FText GeneratedQuery = FText();
 	FText GeneratedContext = FText();
@@ -346,11 +350,20 @@ static FText FindItemName(const FWidgetPath& WidgetPath, FAIAssistantSlateQueryC
 	if (Button.IsValid())
 	{
 		ChildText = FindChildWidgetWithText(Button);
-		if (!ChildText.IsEmpty() || !FindCurrentToolTipText().IsEmpty())
+		if (!ChildText.IsEmpty() || !SlateQueryContext.CurrentToolTipText.IsEmpty())
 		{
-			OutName = (!ChildText.IsEmpty()) ? ChildText : LOCTEXT("ButtonWithNoTextPlaceholderName", "[icon-only]");
+			if (!ChildText.IsEmpty())
+			{
+				OutName = ChildText;
+			}
+			else
+			{
+				// The tool tip can be used in this case.
+				OutName = SlateQueryContext.CurrentToolTipText;
+			}
 			OutDescriptor = LOCTEXT("ItemDescriptor_Button", "button");
 			SlateQueryContext.LastPickedWidget = Button;
+			SlateQueryContext.bIsUIWidget = true;
 		}
 	}
 
@@ -379,11 +392,20 @@ static FText FindItemName(const FWidgetPath& WidgetPath, FAIAssistantSlateQueryC
 	if (ToolbarItemBlock->GetType() != "SNullWidgetContent")
 	{
 		ChildText = FindChildWidgetWithText(ToolbarItemBlock);
-		if (!ChildText.IsEmpty() || !FindCurrentToolTipText().IsEmpty())
+		if (!ChildText.IsEmpty() || !SlateQueryContext.CurrentToolTipText.IsEmpty())
 		{
-			OutName = (!ChildText.IsEmpty()) ? ChildText : LOCTEXT("ButtonWithNoTextPlaceholderName", "[icon-only]");
+			if (!ChildText.IsEmpty())
+			{
+				OutName = ChildText;
+			}
+			else
+			{
+				// The tool tip can be used in this case.
+				OutName = SlateQueryContext.CurrentToolTipText;
+			}
 			OutDescriptor = LOCTEXT("ItemDescriptor_ToolbarButton", "toolbar button");
 			SlateQueryContext.LastPickedWidget = ToolbarItemBlock.ToSharedPtr();
+			SlateQueryContext.bIsUIWidget = true;
 		}
 	}
 
@@ -414,6 +436,7 @@ static FText FindItemName(const FWidgetPath& WidgetPath, FAIAssistantSlateQueryC
 		OutName = FText::Format(LOCTEXT("SearchBoxFormat", "{0}"), (HintText.IsEmpty()) ? LOCTEXT("SearchBoxDefaultName", "default") : HintText);
 		OutDescriptor = LOCTEXT("ItemDescriptor_SearchBox", "search box");
 		SlateQueryContext.LastPickedWidget = SearchBox.ToSharedPtr();
+		SlateQueryContext.bIsUIWidget = true;
 	}
 
 	// Is it a console input box?
@@ -429,6 +452,7 @@ static FText FindItemName(const FWidgetPath& WidgetPath, FAIAssistantSlateQueryC
 		OutName = FText::Format(LOCTEXT("ConsoleInputBoxFormat", "{0}"), (HintText.IsEmpty() ? LOCTEXT("ConsoleInputBoxDefaultName","Console") : HintText));
 		OutDescriptor = LOCTEXT("ItemDescriptor_ConsoleInputBox", "input box");
 		SlateQueryContext.LastPickedWidget = ConsoleInputBox.ToSharedPtr();
+		SlateQueryContext.bIsUIWidget = true;
 	}
 
 	// Is it a menu item? Look up for an SMenuEntryBlock or SWidgetBlock and then down for the text.
@@ -441,6 +465,7 @@ static FText FindItemName(const FWidgetPath& WidgetPath, FAIAssistantSlateQueryC
 			OutName = ChildText;
 			OutDescriptor = LOCTEXT("ItemDescriptor_Menu", "menu item");
 			SlateQueryContext.LastPickedWidget = MenuItemBlock.ToSharedPtr();
+			SlateQueryContext.bIsUIWidget = true;
 		}
 	}
 
@@ -472,6 +497,7 @@ static FText FindItemName(const FWidgetPath& WidgetPath, FAIAssistantSlateQueryC
 			OutName = ChildText;
 			OutDescriptor = LOCTEXT("ItemDescriptor_Breadcrumb", "navigation breadcrumb");
 			SlateQueryContext.LastPickedWidget = BreadcrumbButton.ToSharedPtr();
+			SlateQueryContext.bIsUIWidget = true;
 		}
 	}
 
@@ -479,30 +505,51 @@ static FText FindItemName(const FWidgetPath& WidgetPath, FAIAssistantSlateQueryC
 	TSharedRef<SWidget> AssetTileItem = FindClosestWidgetOfType(WidgetPath, "SAssetTileItem");
 	if (AssetTileItem->GetType() != "SNullWidgetContent")
 	{
-		// TSharedRef<SAssetTileItem> ThisAssetTileItemCast = StaticCastSharedRef<SAssetTileItem>(AssetTileItem);
 		TSharedRef<SWidget> AssetThumbnail = FindClosestWidgetOfType(WidgetPath, "SAssetThumbnail");
 		if (AssetThumbnail->GetType() != "SNullWidgetContent")
 		{
 			ChildText = FindChildWidgetWithText(AssetThumbnail); // This returns TYPE of asset instead of name
-			OutDescriptor = LOCTEXT("ItemDescriptor_Asset", "asset");
 			if (!ChildText.IsEmpty())
 			{
-				OutName = ChildText;
-				SlateQueryContext.LastPickedWidget = AssetThumbnail.ToSharedPtr();
+				FFormatNamedArguments Args;
+				Args.Add(TEXT("AssetType"), ChildText);
+				OutDescriptor = FText::Format(LOCTEXT("ItemDescriptor_Asset", "{AssetType} asset"), Args);
 			}
+
+			if (!SlateQueryContext.CurrentToolTipText.IsEmpty())
+			{
+				// We can use the tooltip for the asset's name.
+				OutName = SlateQueryContext.CurrentToolTipText;
+			}
+			SlateQueryContext.LastPickedWidget = AssetThumbnail.ToSharedPtr();
+			SlateQueryContext.bIsObject = true;
 		}
 		else
 		{
 			ChildText = FindChildWidgetWithText(AssetTileItem);
-			OutDescriptor = LOCTEXT("ItemDescriptor_Folder", "asset folder");
 			if (!ChildText.IsEmpty())
 			{
 				OutName = ChildText;
+				OutDescriptor = LOCTEXT("ItemDescriptor_Folder", "asset folder");
 				SlateQueryContext.LastPickedWidget = AssetTileItem.ToSharedPtr();
+				SlateQueryContext.bIsObject = true;
 			}
 		}
 	}
 
+	// If in Outliner, treat unnamed widgets as actor instances.
+	TSharedRef<SWidget> Outliner = FindClosestWidgetOfType(WidgetPath, "SSceneOutliner");
+	if (Outliner->GetType() != "SNullWidgetContent")
+	{
+		SlateQueryContext.bInOutliner = true;
+	}
+	// If in SubobjectInstanceEditor, treat unnamed widgets as actor instances.
+	TSharedRef<SWidget> SubobjectInstanceEditor = FindClosestWidgetOfType(WidgetPath, "SSubobjectInstanceEditor");
+	if (SubobjectInstanceEditor->GetType() != "SNullWidgetContent")
+	{
+		SlateQueryContext.bInOutliner = true;
+	}
+
 	if (!OutName.IsEmpty())
 	{
 		FFormatNamedArguments Args;
@@ -538,7 +585,7 @@ static FText FindTabName(const FWidgetPath& WidgetPath, FAIAssistantSlateQueryCo
 			}
 		}
 	}
-	
+
 	if (!OutName.IsEmpty())
 	{
 		FFormatNamedArguments Args;
@@ -813,8 +860,10 @@ void UE::AIAssistant::SlateQuerier::QueryAIAssistantAboutSlateWidgetUnderCursor(
 	// We build this up, below.
 	
 	FAIAssistantSlateQueryContext SlateQueryContext;
-	
-	
+
+	// Is there a current tool-tip?
+	SlateQueryContext.CurrentToolTipText = FindCurrentToolTipText();
+
 	// For menu items, we want to generate the context based on the root widget that spawned the menu,
 	// not on the actual menu item itself.
 	const FWidgetPath ContextWidgetPath = FindContextWidgetPath(WidgetPath);
@@ -832,58 +881,79 @@ void UE::AIAssistant::SlateQuerier::QueryAIAssistantAboutSlateWidgetUnderCursor(
 	// Is there an identifiable item?
 	const FText ItemName = FindItemName(WidgetPath, SlateQueryContext);
 
-	// Is there a current tool-tip?
-	const FText ToolTipText = FindCurrentToolTipText();
-
-
+	// Special case for actor listings in Outliner.
+	if (ItemName.IsEmpty() && SlateQueryContext.bInOutliner && !SlateQueryContext.bIsUIWidget)
 	{
-		static const FText QueryAdditionalInstructions = LOCTEXT("QueryInstructions",
-			"If the object is clickable, like a button, then explain what happens if I click on it, including any unrecoverable or unexpected side effects.  If it is safe, don't bother mentioning it. If the object is something you select, like an asset or a blueprint node, describe what the object is or what it does.");
+		SlateQueryContext.bIsObject = true;
+	}
 
-		if (!ItemName.IsEmpty() || !TabName.IsEmpty() || !WindowName.IsEmpty())
+	{
+		FText NameToQuery;
+		bool bHasName = false;
+		if (ItemName.IsEmpty() && !TextUnderCursor.IsEmpty())
 		{
-			const FText NameToQuery = (!ItemName.IsEmpty()) ? ItemName : (!TabName.IsEmpty()) ? TabName : WindowName;
+			NameToQuery = TextUnderCursor;
+		}
+		else if (!ItemName.IsEmpty() || !TabName.IsEmpty() || !WindowName.IsEmpty())
+		{
+			NameToQuery = (!ItemName.IsEmpty()) ? ItemName : (!TabName.IsEmpty()) ? TabName : WindowName;
+			bHasName = true;
+		}
 
-			static const FText QueryFormat = LOCTEXT("QueryFormat1",
-				"I would like to know what {NameToQuery} does in the Unreal Editor UI.");
+		static const FText QueryFormatWidget = LOCTEXT("QueryFormat1", "I would like to know what {NameToQuery} does.");
+		static const FText QueryFormatGeneral = LOCTEXT("QueryFormat2", "I would like to know what \"{NameToQuery}\" means.");
+		static const FText QueryFormatObject = LOCTEXT("QueryFormat3", "I would like to know about {NameToQuery}.");
+		static const FText QueryInstructionsWidgetFormat = LOCTEXT("QueryInstructionsFormat1",
+			"Provide an easily readable, informative, accurate answer that describes what the widget does in the Unreal Editor UI. {QueryAdditionalInstructions} {StyleInstructions}");
+		static const FText QueryInstructionsObjectFormat = LOCTEXT("QueryInstructionsFormat2",
+			"Provide an easily readable, informative, accurate answer that describes what I should know about it in the Unreal Editor UI. {QueryAdditionalInstructions} {StyleInstructions}");
+		static const FText QueryAdditionalInstructionsWidget = LOCTEXT("QueryAdditionalInstructions1",
+			"Explain what happens if I select or click on it.");
+		static const FText QueryAdditionalInstructionsObject = LOCTEXT("QueryAdditionalInstructions2",
+			"Describe what the object is or what it does.");
+		static const FText QueryAdditionalInstructionsGeneral = LOCTEXT("QueryAdditionalInstructions3",
+			"If the object is clickable, like a button, then explain what happens if I click on it. If the object is something you select, like an asset or a blueprint node, describe what the object is or what it does.");
+		static const FText StyleInstructionsWidget = LOCTEXT("StyleInstructions1",
+			"Use bold text to emphasize Unreal Engine specific terms. Start with a quick overview paragraph, then add key points in one or two formatted sections with headers and bullet points. Don't include a summary at the end.");
+		static const FText StyleInstructionsObject = LOCTEXT("StyleInstructions2",
+			"Use bold text to emphasize Unreal Engine specific terms. Start with a quick overview paragraph, then add key points in one or two formatted sections with headers and bullet points. Don't include a summary at the end.");
+
+		if (!NameToQuery.IsEmpty())
+		{
 			FFormatNamedArguments QueryArgs;
 			QueryArgs.Add(TEXT("NameToQuery"), NameToQuery);
-			SlateQueryContext.GeneratedQuery = FText::Format(QueryFormat, QueryArgs);
-			
-			static const FText QueryInstructionsFormat = LOCTEXT("QueryInstructionsFormat1",
-				"Provide a short, conversational answer that describes what the widget does. {QueryAdditionalInstructions}");
 			FFormatNamedArguments InstructionArgs;
-			InstructionArgs.Add(TEXT("QueryAdditionalInstructions"), QueryAdditionalInstructions);
-			SlateQueryContext.GeneratedQueryInstructions = FText::Format(QueryInstructionsFormat, InstructionArgs);
-		}
-		
-		if (ItemName.IsEmpty())
-		{
-			if (!TextUnderCursor.IsEmpty())
+
+			if (SlateQueryContext.bIsObject)
 			{
-				static const FText QueryFormat = LOCTEXT("QueryFormat2",
-					"What does \"{ObjectText}\" mean in the Unreal Editor UI?");
-				FFormatNamedArguments Args;
-				Args.Add(TEXT("ObjectText"), TextUnderCursor);
-				SlateQueryContext.GeneratedQuery = FText::Format(QueryFormat, Args);
-				
-				static const FText QueryInstructionsFormat = LOCTEXT("QueryInstructionsFormat2",
-					"Provide a short, conversational answer that describes what I should know about it. {QueryAdditionalInstructions}");
-				FFormatNamedArguments InstructionArgs;
-				InstructionArgs.Add(TEXT("QueryAdditionalInstructions"), QueryAdditionalInstructions);
-				SlateQueryContext.GeneratedQueryInstructions = FText::Format(QueryInstructionsFormat, InstructionArgs);
+				SlateQueryContext.GeneratedQuery = FText::Format(QueryFormatObject, QueryArgs);
+				InstructionArgs.Add(TEXT("QueryAdditionalInstructions"), QueryAdditionalInstructionsObject);
+				InstructionArgs.Add(TEXT("StyleInstructions"), StyleInstructionsObject);
+				SlateQueryContext.GeneratedQueryInstructions = FText::Format(QueryInstructionsObjectFormat, InstructionArgs);
+			}
+			else if (SlateQueryContext.bIsUIWidget)
+			{
+				SlateQueryContext.GeneratedQuery = FText::Format(QueryFormatWidget, QueryArgs);
+				InstructionArgs.Add(TEXT("QueryAdditionalInstructions"), QueryAdditionalInstructionsWidget);
+				InstructionArgs.Add(TEXT("StyleInstructions"), StyleInstructionsWidget);
+				SlateQueryContext.GeneratedQueryInstructions = FText::Format(QueryInstructionsWidgetFormat, InstructionArgs);
+			}
+			else
+			{
+				SlateQueryContext.GeneratedQuery = FText::Format(QueryFormatWidget, QueryArgs);
+				InstructionArgs.Add(TEXT("QueryAdditionalInstructions"), QueryAdditionalInstructionsGeneral);
+				InstructionArgs.Add(TEXT("StyleInstructions"), StyleInstructionsWidget);
+				SlateQueryContext.GeneratedQueryInstructions = FText::Format(QueryInstructionsWidgetFormat, InstructionArgs);
 			}
 		}
 
 		if (SlateQueryContext.GeneratedQuery.IsEmpty())
 		{
 			UE_LOG(LogAIAssistant, Warning, TEXT("Could not generate query for widget."));
-			
 			return;
 		}
 	}
 
-
 	{
 		if (!TabName.IsEmpty() || !WindowName.IsEmpty())
 		{
@@ -914,11 +984,11 @@ void UE::AIAssistant::SlateQuerier::QueryAIAssistantAboutSlateWidgetUnderCursor(
 			SlateQueryContext.GeneratedContextItems.Add(DetailsViewContext);
 		}
 
-		if (!ToolTipText.IsEmpty())
+		if (!SlateQueryContext.CurrentToolTipText.IsEmpty())
 		{
-			FText ToolTipContextFormat = LOCTEXT("QueryContextToolTipAddition", "The text of the current tooltip is: \"{ToolTipText}\"");
+			FText ToolTipContextFormat = LOCTEXT("QueryContextToolTipAddition", "The text of the current tooltip is: \"{ToolTipText}\", but you don't need to mention it.");
 			FFormatNamedArguments Args;
-			Args.Add(TEXT("ToolTipText"), ToolTipText);
+			Args.Add(TEXT("ToolTipText"), SlateQueryContext.CurrentToolTipText);
 			FText ToolTipContext = FText::Format(ToolTipContextFormat, Args);
 			SlateQueryContext.GeneratedContextItems.Add(ToolTipContext);
 		}
