[AI Assistant] Conversation ready executor.

ExecuteWhenReady implementation that enqueues executions until a conversation
is ready.

diff --git a/Engine/Plugins/Experimental/AIAssistant/Source/AIAssistant/Private/AIAssistantConversationReadyExecutor.cpp b/Engine/Plugins/Experimental/AIAssistant/Source/AIAssistant/Private/AIAssistantConversationReadyExecutor.cpp
new file mode 100644
--- /dev/null
+++ b/Engine/Plugins/Experimental/AIAssistant/Source/AIAssistant/Private/AIAssistantConversationReadyExecutor.cpp
@@ -0,0 +1,75 @@
+// Copyright Epic Games, Inc. All Rights Reserved.
+
+#include "AIAssistantConversationReadyExecutor.h"
+
+#include "Async/Mutex.h"
+#include "Async/UniqueLock.h"
+#include "Templates/Function.h"
+#include "Templates/UnrealTemplate.h"
+
+namespace UE::AIAssistant
+{
+	FConversationReadyExecutor::FConversationReadyExecutor(
+		FGetExecuteWhenReadyStateFunc&& GetExecuteWhenReadyStateFunc) :
+			GetExecuteWhenReadyStateDelegate(MoveTemp(GetExecuteWhenReadyStateFunc))
+	{
+	}
+
+	// Notify the executor that the agent environment has been configured.
+	void FConversationReadyExecutor::NotifyAgentEnvironmentConfigured()
+	{
+		{
+			UE::TUniqueLock Lock(StateMutex);
+			bConfiguredAgentEnvironment = true;
+		}
+		UpdateExecuteWhenReady();
+	}
+
+	bool FConversationReadyExecutor::SetCreatingConversation(bool bNewCreatingConversation)
+	{
+		bool bPreviousCreatingConversation;
+		bool bCreatingNewConversation;
+		bool bCreatedNewConversation;
+		{
+			UE::TUniqueLock Lock(StateMutex);
+			bPreviousCreatingConversation = bCreatingConversation;
+			bCreatingConversation = bNewCreatingConversation;
+			bCreatingNewConversation = !bPreviousCreatingConversation && bCreatingConversation;
+			bCreatedNewConversation = bPreviousCreatingConversation && !bCreatingConversation;
+		}
+		if (bCreatingNewConversation)
+		{
+			ResetExecuteWhenReady();
+		}
+		if (bCreatedNewConversation)
+		{
+			UpdateExecuteWhenReady();
+		}
+		return bPreviousCreatingConversation;
+	}
+
+	FExecuteWhenReady::EExecuteWhenReadyState FConversationReadyExecutor::GetExecuteWhenReadyState()
+	{
+		if (GetExecuteWhenReadyStateDelegate)
+		{
+			auto ExecuteWhenReadyState = GetExecuteWhenReadyStateDelegate();
+			if (ExecuteWhenReadyState != EExecuteWhenReadyState::Execute)
+			{
+				return ExecuteWhenReadyState;
+			}
+		}
+		UE::TUniqueLock Lock(StateMutex);
+		return !bConfiguredAgentEnvironment || bCreatingConversation
+			? EExecuteWhenReadyState::Wait
+			: EExecuteWhenReadyState::Execute;
+	}
+
+	void FConversationReadyExecutor::ClearExecutionQueueIfCreatingConversation()
+	{
+		UE::TUniqueLock Lock(StateMutex);
+		if (bCreatingConversation)
+		{
+			ResetExecuteWhenReady();
+		}
+	}
+}
\ No newline at end of file
diff --git a/Engine/Plugins/Experimental/AIAssistant/Source/AIAssistant/Private/AIAssistantConversationReadyExecutor.h b/Engine/Plugins/Experimental/AIAssistant/Source/AIAssistant/Private/AIAssistantConversationReadyExecutor.h
new file mode 100644
--- /dev/null
+++ b/Engine/Plugins/Experimental/AIAssistant/Source/AIAssistant/Private/AIAssistantConversationReadyExecutor.h
@@ -0,0 +1,52 @@
+// Copyright Epic Games, Inc. All Rights Reserved.
+#pragma once
+
+#include "Async/Mutex.h"
+#include "Templates/Function.h"
+#include "Templates/UnrealTemplate.h"
+
+#include "AIAssistantExecuteWhenReady.h"
+
+namespace UE::AIAssistant
+{
+	// Handles enqueuing of operations against a conversation until it's ready to accept messages.
+	class FConversationReadyExecutor : public FExecuteWhenReady
+	{
+	public:
+		using FGetExecuteWhenReadyStateFunc = TFunction<FExecuteWhenReady::EExecuteWhenReadyState()>;
+
+	public:
+		// Construct with a delegate that retrieves additional execution when ready state.
+		FConversationReadyExecutor(
+			FGetExecuteWhenReadyStateFunc&& GetExecuteWhenReadyStateDelegate =
+				FGetExecuteWhenReadyStateFunc());
+
+		virtual ~FConversationReadyExecutor() = default;
+
+		// Notify the executor that the agent environment has been configured.
+		void NotifyAgentEnvironmentConfigured();
+
+		// Set the creating conversation flag returning the previous value.
+		bool SetCreatingConversation(bool bNewCreatingConversation);
+
+		template<class CallableType>
+		void ExecuteWhenReady(CallableType&& Callable)
+		{
+			ClearExecutionQueueIfCreatingConversation();
+			FExecuteWhenReady::ExecuteWhenReady(Forward<CallableType>(Callable));
+		}
+
+	protected:
+		FExecuteWhenReady::EExecuteWhenReadyState GetExecuteWhenReadyState() override;
+
+	private:
+		void ClearExecutionQueueIfCreatingConversation();
+
+	private:
+		mutable UE::FMutex StateMutex;
+		bool bConfiguredAgentEnvironment = false;
+		bool bCreatingConversation = false;
+		FGetExecuteWhenReadyStateFunc GetExecuteWhenReadyStateDelegate;
+
+	};
+}
\ No newline at end of file
diff --git a/Engine/Plugins/Experimental/AIAssistant/Source/AIAssistant/Private/Tests/AIAssistantConversationReadyExecutorTest.cpp b/Engine/Plugins/Experimental/AIAssistant/Source/AIAssistant/Private/Tests/AIAssistantConversationReadyExecutorTest.cpp
new file mode 100644
--- /dev/null
+++ b/Engine/Plugins/Experimental/AIAssistant/Source/AIAssistant/Private/Tests/AIAssistantConversationReadyExecutorTest.cpp
@@ -0,0 +1,144 @@
+// Copyright Epic Games, Inc. All Rights Reserved.
+
+#include "Containers/UnrealString.h"
+#include "Misc/AutomationTest.h"
+
+#include "AIAssistantConversationReadyExecutor.h"
+#include "AIAssistantExecuteWhenReady.h"
+#include "AIAssistantTestFlags.h"
+
+using namespace UE::AIAssistant;
+
+#if WITH_DEV_AUTOMATION_TESTS
+
+IMPLEMENT_SIMPLE_AUTOMATION_TEST(
+	FAIAssistantConversationReadyExecutorSetCreatingConversationTest,
+	"AI.Assistant.ConversationReadyExecutor.SetCreatingConversation",
+	AIAssistantTest::Flags);
+
+bool FAIAssistantConversationReadyExecutorSetCreatingConversationTest::RunTest(
+	const FString& UnusedParameters)
+{
+	FConversationReadyExecutor ConversationReadyExecutor;
+	(void)TestFalse(
+		TEXT("No change to conversation state"),
+		ConversationReadyExecutor.SetCreatingConversation(false));
+	(void)TestFalse(
+		TEXT("Creating conversation"),
+		ConversationReadyExecutor.SetCreatingConversation(true));
+	(void)TestTrue(
+		TEXT("Still creating conversation"),
+		ConversationReadyExecutor.SetCreatingConversation(true));
+	(void)TestTrue(
+		TEXT("Created conversation"),
+		ConversationReadyExecutor.SetCreatingConversation(false));
+	(void)TestFalse(
+		TEXT("No change to conversation state after creation"),
+		ConversationReadyExecutor.SetCreatingConversation(false));
+	return true;
+}
+
+IMPLEMENT_SIMPLE_AUTOMATION_TEST(
+	FAIAssistantConversationReadyExecutorExecuteOnConstructionTest,
+	"AI.Assistant.ConversationReadyExecutor.ExecuteOnConstruction",
+	AIAssistantTest::Flags);
+
+bool FAIAssistantConversationReadyExecutorExecuteOnConstructionTest::RunTest(
+	const FString& UnusedParameters)
+{
+	FConversationReadyExecutor ConversationReadyExecutor;
+	bool bExecuted = false;
+	ConversationReadyExecutor.ExecuteWhenReady([&bExecuted]() -> void { bExecuted = true; });
+	(void)TestFalse(TEXT("Not configured on construction"), bExecuted);
+
+	ConversationReadyExecutor.NotifyAgentEnvironmentConfigured();
+	(void)TestTrue(TEXT("Ready when configured"), bExecuted);
+	return true;
+}
+
+IMPLEMENT_SIMPLE_AUTOMATION_TEST(
+	FAIAssistantConversationReadyExecutorExecuteAfterConversationCreatedTest,
+	"AI.Assistant.ConversationReadyExecutor.ExecuteAfterConversationCreated",
+	AIAssistantTest::Flags);
+
+bool FAIAssistantConversationReadyExecutorExecuteAfterConversationCreatedTest::RunTest(
+	const FString& UnusedParameters)
+{
+	FConversationReadyExecutor ConversationReadyExecutor;
+	ConversationReadyExecutor.NotifyAgentEnvironmentConfigured();
+	(void)ConversationReadyExecutor.SetCreatingConversation(true);
+	bool bExecuted = false;
+	ConversationReadyExecutor.ExecuteWhenReady([&bExecuted]() -> void { bExecuted = true; });
+	(void)TestFalse(TEXT("Not executed while creating conversation"), bExecuted);
+
+	(void)ConversationReadyExecutor.SetCreatingConversation(false);
+	(void)TestTrue(TEXT("Executed after conversation created"), bExecuted);
+	return true;
+}
+
+
+IMPLEMENT_SIMPLE_AUTOMATION_TEST(
+	FAIAssistantConversationReadyExecutorOnlyExecuteLastAfterConversationCreatedTest,
+	"AI.Assistant.ConversationReadyExecutor.OnlyExecuteLastAfterConversationCreated",
+	AIAssistantTest::Flags);
+
+bool FAIAssistantConversationReadyExecutorOnlyExecuteLastAfterConversationCreatedTest::RunTest(
+	const FString& UnusedParameters)
+{
+	FConversationReadyExecutor ConversationReadyExecutor;
+	ConversationReadyExecutor.NotifyAgentEnvironmentConfigured();
+	(void)ConversationReadyExecutor.SetCreatingConversation(true);
+	int ExecutedItem = 0;
+	ConversationReadyExecutor.ExecuteWhenReady([&ExecutedItem]() -> void { ExecutedItem = 1; });
+	ConversationReadyExecutor.ExecuteWhenReady([&ExecutedItem]() -> void { ExecutedItem = 2; });
+	ConversationReadyExecutor.ExecuteWhenReady([&ExecutedItem]() -> void { ExecutedItem = 3; });
+	(void)TestEqual(TEXT("Not executed while creating conversation"), ExecutedItem, 0);
+
+	(void)ConversationReadyExecutor.SetCreatingConversation(false);
+	(void)TestEqual(TEXT("Executed after conversation created"), ExecutedItem, 3);
+	return true;
+}
+
+IMPLEMENT_SIMPLE_AUTOMATION_TEST(
+	FAIAssistantConversationReadyExecutorDelegateReadyStateTest,
+	"AI.Assistant.ConversationReadyExecutor.DelegateReadyState",
+	AIAssistantTest::Flags);
+
+bool FAIAssistantConversationReadyExecutorDelegateReadyStateTest::RunTest(
+	const FString& UnusedParameters)
+{
+	FExecuteWhenReady::EExecuteWhenReadyState ExecuteWhenReadyState =
+		FExecuteWhenReady::EExecuteWhenReadyState::Wait;
+	FConversationReadyExecutor ConversationReadyExecutor(
+		[&ExecuteWhenReadyState]() -> FExecuteWhenReady::EExecuteWhenReadyState
+		{
+			return ExecuteWhenReadyState;
+		});
+	ConversationReadyExecutor.NotifyAgentEnvironmentConfigured();
+
+	const FString LastExecutionNothing = TEXT("Nothing");
+	FString LastExecution = LastExecutionNothing;
+	ConversationReadyExecutor.ExecuteWhenReady(
+		[&LastExecution]() -> void { LastExecution = TEXT("First"); });
+	(void)TestEqual(TEXT("Should wait on the delegate"), LastExecution, LastExecutionNothing);
+
+	ExecuteWhenReadyState = FExecuteWhenReady::EExecuteWhenReadyState::Reject;
+	ConversationReadyExecutor.UpdateExecuteWhenReady();
+	(void)TestEqual(
+		TEXT("Pending execution should be cleared from the queue."),
+		LastExecution, LastExecutionNothing);
+
+	ExecuteWhenReadyState = FExecuteWhenReady::EExecuteWhenReadyState::Execute;
+	const FString LastExecutionLast = TEXT("Nothing");
+	ConversationReadyExecutor.ExecuteWhenReady(
+		[&LastExecution, &LastExecutionLast]() -> void
+		{
+			LastExecution = LastExecutionLast;
+		});
+	(void)TestEqual(
+		TEXT("Only the most recent function should be executed."),
+		LastExecution, LastExecutionLast);
+	return true;
+}
+
+#endif
\ No newline at end of file
