[AI Assistant] Support for new conversation API

Add support for conversation API that returns a promise.

diff --git a/Engine/Plugins/Experimental/AIAssistant/Source/AIAssistant/Private/AIAssistantWebApi.cpp b/Engine/Plugins/Experimental/AIAssistant/Source/AIAssistant/Private/AIAssistantWebApi.cpp
--- a/Engine/Plugins/Experimental/AIAssistant/Source/AIAssistant/Private/AIAssistantWebApi.cpp
+++ b/Engine/Plugins/Experimental/AIAssistant/Source/AIAssistant/Private/AIAssistantWebApi.cpp
@@ -55,9 +55,9 @@ try {
 		(void)ExecuteFunctionWithJsonArgument(TEXT("addMessageToConversation"), Options);
 	}
 
-	void FWebApi::CreateConversation()
+	TFuture<TValueOrError<void, FString>> FWebApi::CreateConversation()
 	{
-		(void)ExecuteFunction(TEXT("createConversation"));
+		return ExecutionFunctionParseJson<void>(TEXT("createConversation"));
 	}
 
 	TFuture<TValueOrError<FAgentEnvironmentHandle, FString>> FWebApi::AddAgentEnvironment(
diff --git a/Engine/Plugins/Experimental/AIAssistant/Source/AIAssistant/Private/AIAssistantWebApi.h b/Engine/Plugins/Experimental/AIAssistant/Source/AIAssistant/Private/AIAssistantWebApi.h
--- a/Engine/Plugins/Experimental/AIAssistant/Source/AIAssistant/Private/AIAssistantWebApi.h
+++ b/Engine/Plugins/Experimental/AIAssistant/Source/AIAssistant/Private/AIAssistantWebApi.h
@@ -9,8 +9,10 @@
 #include "Misc/Variant.h"
 #include "Serialization/JsonSerializable.h"
 #include "Serialization/JsonSerializerMacros.h"
+#include "Templates/Function.h"
 #include "Templates/SharedPointer.h"
 #include "Templates/Tuple.h"
+#include "Templates/UnrealTemplate.h"
 #include "Templates/ValueOrError.h"
 #include "UObject/StrongObjectPtr.h"
 
@@ -217,7 +219,7 @@ namespace UE::AIAssistant
 		void AddMessageToConversation(const FAddMessageToConversationOptions& Options);
 
 		// Create a new conversation.
-		void CreateConversation();
+		TFuture<TValueOrError<void, FString>> CreateConversation();
 
 		// Add an agent environment for the currently logged in user
 		// returning the ID. If a matching environment already exists for the
@@ -257,39 +259,77 @@ namespace UE::AIAssistant
 			return ExecuteFunction(FunctionName, *Argument.ToJson(false));
 		}
 
-		// Execute a function marshalling an argument to JSON and the return value from JSON or
-		// setting an error.
-		template<typename JsonSerializableReturnType, typename JsonSerializableArgType>
-		TFuture<TValueOrError<JsonSerializableReturnType, FString>>
-			ExecutionFunctionParseJson(const TCHAR* FunctionName, const JsonSerializableArgType& Argument)
+		// Create a promise and handler for the execution of a JavaScript function that optionally
+		// parses a JSON result value.
+		template<typename JsonSerializableReturnType>
+		TTuple<
+			TSharedPtr<TPromise<TValueOrError<JsonSerializableReturnType, FString>>>, 
+			TFunction<void(const TFuture<UAIAssistantWebJavaScriptResultDelegate::FResult>&)>>
+			CreatePromiseAndHandlerForFunction()
 		{
-			auto Result = MakeShared<TPromise<TValueOrError<JsonSerializableReturnType, FString>>>();
-			ExecuteFunctionWithJsonArgument(FunctionName, Argument).Then(
+			auto Result =
+				MakeShared<TPromise<TValueOrError<JsonSerializableReturnType, FString>>>();
+			auto ResultHandler =
 				[Result](
 					const TFuture<UAIAssistantWebJavaScriptResultDelegate::FResult>&
-						ExecutionResultFuture) mutable -> void
+					ExecutionResultFuture) mutable -> void
 				{
 					auto ExecutionResult = ExecutionResultFuture.Get();
 					const FString& Json = ExecutionResult.Json;
-					bool bIsError = ExecutionResult.bJsonIsError;
 					if (ExecutionResult.bJsonIsError)
 					{
 						Result->SetValue(MakeError(Json));
 					}
 					else
 					{
-						JsonSerializableReturnType Parsed;
-						if (!Parsed.FromJson(Json))
-						{
-							Result->SetValue(MakeError(FString(TEXT("Failed to parse: ")) + Json));
-						}
-						else
-						{
-							Result->SetValue(MakeValue(MoveTemp(Parsed)));
-						}
+						Result->SetValue(ParseJsonIfNotVoid<JsonSerializableReturnType>(Json));
 					}
-				});
-			return Result->GetFuture();
+				};
+			return MakeTuple(Result, MoveTemp(ResultHandler));
+		}
+
+		// Execute a function with no arguments optionally marshalling a return value from JSON or
+		// setting an error.
+		template<typename JsonSerializableReturnType>
+		TFuture<TValueOrError<JsonSerializableReturnType, FString>>
+			ExecutionFunctionParseJson(const TCHAR* FunctionName)
+		{
+			auto ResultAndHandler = CreatePromiseAndHandlerForFunction<JsonSerializableReturnType>();
+			ExecuteFunction(FunctionName).Then(MoveTemp(ResultAndHandler.Value));
+			return ResultAndHandler.Key->GetFuture();
+		}
+
+		// Execute a function marshalling an argument to JSON and the return value from JSON or
+		// setting an error.
+		template<typename JsonSerializableReturnType, typename JsonSerializableArgType>
+		TFuture<TValueOrError<JsonSerializableReturnType, FString>>
+			ExecutionFunctionParseJson(
+				const TCHAR* FunctionName,
+				const JsonSerializableArgType& Argument)
+		{
+			auto ResultAndHandler = CreatePromiseAndHandlerForFunction<JsonSerializableReturnType>();
+			ExecuteFunctionWithJsonArgument(FunctionName, Argument).Then(
+				MoveTemp(ResultAndHandler.Value));
+			return ResultAndHandler.Key->GetFuture();
+		}
+
+		// Parse a JSON string if JsonSerializableReturnType is not void otherwise return an empty value.
+		template<typename JsonSerializableReturnType>
+		static TValueOrError<JsonSerializableReturnType, FString> ParseJsonIfNotVoid(const FString& Json)
+		{
+			JsonSerializableReturnType Parsed;
+			if (!Parsed.FromJson(Json))
+			{
+				return MakeError(FString(TEXT("Failed to parse: ")) + Json);
+			}
+			return MakeValue(MoveTemp(Parsed));
+		}
+
+		// Overload that ignores the supplied JSON string and returns an empty value.
+		template<>
+		inline TValueOrError<void, FString> ParseJsonIfNotVoid(const FString& UnusedJson)
+		{
+			return MakeValue();
 		}
 
 	private:
diff --git a/Engine/Plugins/Experimental/AIAssistant/Source/AIAssistant/Private/Tests/AIAssistantWebApiTest.cpp b/Engine/Plugins/Experimental/AIAssistant/Source/AIAssistant/Private/Tests/AIAssistantWebApiTest.cpp
--- a/Engine/Plugins/Experimental/AIAssistant/Source/AIAssistant/Private/Tests/AIAssistantWebApiTest.cpp
+++ b/Engine/Plugins/Experimental/AIAssistant/Source/AIAssistant/Private/Tests/AIAssistantWebApiTest.cpp
@@ -51,6 +51,20 @@ public:
 		return bCalled;
 	}
 
+	// If possible, get the value of a result as JSON.  If the value cannot be converted to JSON
+	// return an empty string.
+	template<typename ValueType>
+	FString GetValueOrEmptyString(const TValueOrError<ValueType, FString>& ValueOrError)
+	{
+		return ValueOrError.GetValue().ToJson(false);
+	}
+
+	template<>
+	FString GetValueOrEmptyString<void>(const TValueOrError<void, FString>& UnusedValueOrError)
+	{
+		return FString();
+	}
+
 	// Ensure a JavaScript function was called with the specified arguments and complete it
 	// with ResultJson as a result of error if bResultJsonIsError is true.
 	template<typename ResultFutureType>
@@ -88,7 +102,7 @@ public:
 			Result.HasError(), bResultJsonIsError);
 		bCompletedAsExpected = TestCase.TestEqual(
 			FString::Printf(TEXT("%s_HasExpectedResult"), FunctionName),
-			bResultJsonIsError ? Result.GetError() : Result.GetValue().ToJson(false),
+			bResultJsonIsError ? Result.GetError() : GetValueOrEmptyString(Result),
 			ResultJson) && bCompletedAsExpected;
 		return bCompletedAsExpected;
 	}
@@ -265,8 +279,9 @@ IMPLEMENT_SIMPLE_AUTOMATION_TEST(
 bool FAIAssistantWebApiTestCreateConversation::RunTest(const FString& UnusedParameters)
 {
 	FFakeWebApi WebApi;
-	WebApi->CreateConversation();
-	return WebApi->TestExpectAsyncFunctionCall(*this, TEXT("createConversation"));
+	auto Result = WebApi->CreateConversation();
+	return WebApi->TestExpectAsyncFunctionCallAndComplete(
+		*this, TEXT("createConversation"), TEXT(""), Result, TEXT(""), false);
 }
 
 IMPLEMENT_SIMPLE_AUTOMATION_TEST(
