[AI Assistant] Send structured context for UI tooltips
- improve item detection for
- asset tree items
- asset thumbnails
- outliner items
- slight tweak to prompt instructions to limit verbosity

diff --git a/Engine/Plugins/Experimental/AIAssistant/Source/AIAssistant/Private/AIAssistantSlateQuerier.cpp b/Engine/Plugins/Experimental/AIAssistant/Source/AIAssistant/Private/AIAssistantSlateQuerier.cpp
--- a/Engine/Plugins/Experimental/AIAssistant/Source/AIAssistant/Private/AIAssistantSlateQuerier.cpp
+++ b/Engine/Plugins/Experimental/AIAssistant/Source/AIAssistant/Private/AIAssistantSlateQuerier.cpp
@@ -9,6 +9,8 @@
 #include "EditorModeManager.h"
 #include "IDetailsView.h"
 #include "Interfaces/IMainFrameModule.h"
+#include "Serialization/JsonSerializable.h"
+#include "Serialization/JsonSerializerMacros.h"
 #include "Subsystems/AssetEditorSubsystem.h"
 #include "Toolkits/BaseToolkit.h"
 #include "Widgets/Text/STextBlock.h"
@@ -28,6 +30,39 @@
 
 #define LOCTEXT_NAMESPACE "FAIAssistantSlateQuerier"
 
+struct FAIAssistantSlateQueryStructuredContext : public FJsonSerializable
+{
+	FString TopWindowName;
+	TOptional<FString> WindowName;
+	TOptional<FString> EditorName;
+	TOptional<FString> EditorMode;
+	TOptional<FString> ActiveToolName;
+	TOptional<FString> TabName;
+	TOptional<FString> ToolTipText;
+	TOptional<FString> TextUnderCursor;
+	TOptional<FString> WidgetType;
+	TOptional<FString> ActorClass;
+	TOptional<FString> ItemName;
+	TOptional<FString> ItemDescriptor;
+	TOptional<TArray<FString>> WidgetPath;
+
+	BEGIN_JSON_SERIALIZER
+		JSON_SERIALIZE("toplevel_window_name", TopWindowName);
+		JSON_SERIALIZE_OPTIONAL("window_name", WindowName);
+		JSON_SERIALIZE_OPTIONAL("editor_name", EditorName);
+		JSON_SERIALIZE_OPTIONAL("editor_mode", EditorMode);
+		JSON_SERIALIZE_OPTIONAL("active_tool", ActiveToolName);
+		JSON_SERIALIZE_OPTIONAL("tab_name", TabName);
+		JSON_SERIALIZE_OPTIONAL("tooltip_text", ToolTipText);
+		JSON_SERIALIZE_OPTIONAL("text_under_cursor", TextUnderCursor);
+		JSON_SERIALIZE_OPTIONAL("widget_type", WidgetType);
+		JSON_SERIALIZE_OPTIONAL("actor_class", ActorClass);
+		JSON_SERIALIZE_OPTIONAL("item_name", ItemName);
+		JSON_SERIALIZE_OPTIONAL("item_descriptor", ItemDescriptor);
+		JSON_SERIALIZE_OPTIONAL_ARRAY("widget_path", WidgetPath);
+	END_JSON_SERIALIZER
+};
+
 
 //
 // FAIAssistantSlateQueryContext
@@ -39,15 +74,20 @@
 struct FAIAssistantSlateQueryContext
 {
 	TSharedPtr<SWidget> LastPickedWidget;
-	TArray<FText> GeneratedContextItems;
 	FText CurrentToolTipText = FText();
+	FText FormattedWindowName = FText();
+	FText FormattedTabName = FText();
+	FText TextUnderCursor = FText();
+	FText ItemName = FText();
+	FText ItemDescriptor = FText();
 	bool bIsUIWidget = false;
 	bool bIsObject = false;
 	bool bInOutliner = false;
 
 	FText GeneratedQuery = FText();
-	FText GeneratedContext = FText();
 	FText GeneratedQueryInstructions = FText();
+
+	FAIAssistantSlateQueryStructuredContext StructuredContext;
 };
 
 
@@ -233,10 +273,8 @@ static FWidgetPath FindContextWidgetPath(const FWidgetPath& WidgetPath)
 }
 
 
-static FText GenerateModeToolContext(const FWidgetPath& WidgetPath)
+static void GenerateModeToolContext(const FWidgetPath& WidgetPath, FAIAssistantSlateQueryContext& SlateQueryContext)
 {
-	FText OutContext = FText();
-
 	// if we're in a level editor mode, include which tool is active
 	TSharedRef<SWidget> LevelEditor = FindClosestWidgetOfType(WidgetPath, "SLevelEditor");
 	if (LevelEditor->GetType() != "SNullWidgetContent")
@@ -258,25 +296,20 @@ static FText GenerateModeToolContext(const FWidgetPath& WidgetPath)
 						FText ToolDisplayName = ModeToolkit.Pin()->GetActiveToolDisplayName();
 						if (!ToolDisplayName.IsEmpty())
 						{
-							FFormatNamedArguments Args;
-							Args.Add(TEXT("ToolName"), ToolDisplayName);
-							Args.Add(TEXT("ModeName"), ModeInfo.Name);
-							return FText::Format(LOCTEXT("EditorName_ModeToolFormat", "The {ToolName} tool is active for {ModeName} mode."), Args);
+							SlateQueryContext.StructuredContext.ActiveToolName = ToolDisplayName.ToString();
+							SlateQueryContext.StructuredContext.EditorMode = ModeInfo.Name.ToString();
 						}
 					}
 				}
 			}
 		}
 	}
-
-	return OutContext;
 }
 
 
-static FText FindItemName(const FWidgetPath& WidgetPath, FAIAssistantSlateQueryContext& SlateQueryContext)
+static bool FindItemName(const FWidgetPath& WidgetPath, FAIAssistantSlateQueryContext& SlateQueryContext)
 {
-	FText OutName = FText();
-	FText OutDescriptor = LOCTEXT("ItemDescriptor_Generic", "control");
+	SlateQueryContext.ItemDescriptor = LOCTEXT("ItemDescriptor_Generic", "control");
 	TSet<FName> ButtonTypes = TSet<FName>({
 		"SButton",
 		"SPrimaryButton",
@@ -287,7 +320,7 @@ static FText FindItemName(const FWidgetPath& WidgetPath, FAIAssistantSlateQueryC
 
 	if (!WidgetPath.IsValid())
 	{
-		return OutName;
+		return false;
 	}
 
 	// Is it a graph node?
@@ -315,9 +348,9 @@ static FText FindItemName(const FWidgetPath& WidgetPath, FAIAssistantSlateQueryC
 						continue;
 					}
 					const TSharedPtr<SGraphNode> AsGraphNode = StaticCastSharedPtr<SGraphNode>(ThisNodeWidget.ToSharedPtr());
-					OutName = AsGraphNode->GetNodeObj()->GetNodeTitle(ENodeTitleType::MenuTitle);
-					OutDescriptor = LOCTEXT("ItemDescriptor_GraphNode", "graph node");
 					SlateQueryContext.LastPickedWidget = ThisNodeWidget.ToSharedPtr();
+					SlateQueryContext.ItemName = AsGraphNode->GetNodeObj()->GetNodeTitle(ENodeTitleType::MenuTitle);
+					SlateQueryContext.ItemDescriptor = LOCTEXT("ItemDescriptor_GraphNode", "graph node");
 				}
 			}
 		}
@@ -354,14 +387,14 @@ static FText FindItemName(const FWidgetPath& WidgetPath, FAIAssistantSlateQueryC
 		{
 			if (!ChildText.IsEmpty())
 			{
-				OutName = ChildText;
+				SlateQueryContext.ItemName = ChildText;
 			}
 			else
 			{
 				// The tool tip can be used in this case.
-				OutName = SlateQueryContext.CurrentToolTipText;
+				SlateQueryContext.ItemName = SlateQueryContext.CurrentToolTipText;
 			}
-			OutDescriptor = LOCTEXT("ItemDescriptor_Button", "button");
+			SlateQueryContext.ItemDescriptor = LOCTEXT("ItemDescriptor_Button", "button");
 			SlateQueryContext.LastPickedWidget = Button;
 			SlateQueryContext.bIsUIWidget = true;
 		}
@@ -396,14 +429,14 @@ static FText FindItemName(const FWidgetPath& WidgetPath, FAIAssistantSlateQueryC
 		{
 			if (!ChildText.IsEmpty())
 			{
-				OutName = ChildText;
+				SlateQueryContext.ItemName = ChildText;
 			}
 			else
 			{
 				// The tool tip can be used in this case.
-				OutName = SlateQueryContext.CurrentToolTipText;
+				SlateQueryContext.ItemName = SlateQueryContext.CurrentToolTipText;
 			}
-			OutDescriptor = LOCTEXT("ItemDescriptor_ToolbarButton", "toolbar button");
+			SlateQueryContext.ItemDescriptor = LOCTEXT("ItemDescriptor_ToolbarButton", "toolbar button");
 			SlateQueryContext.LastPickedWidget = ToolbarItemBlock.ToSharedPtr();
 			SlateQueryContext.bIsUIWidget = true;
 		}
@@ -433,8 +466,9 @@ static FText FindItemName(const FWidgetPath& WidgetPath, FAIAssistantSlateQueryC
 	}
 	if (SearchBox->GetType() != "SNullWidgetContent")
 	{
-		OutName = FText::Format(LOCTEXT("SearchBoxFormat", "{0}"), (HintText.IsEmpty()) ? LOCTEXT("SearchBoxDefaultName", "default") : HintText);
-		OutDescriptor = LOCTEXT("ItemDescriptor_SearchBox", "search box");
+		SlateQueryContext.ItemName = FText::Format(LOCTEXT("SearchBoxFormat", "{0}"),
+			(HintText.IsEmpty()) ? LOCTEXT("SearchBoxDefaultName", "default") : HintText);
+		SlateQueryContext.ItemDescriptor = LOCTEXT("ItemDescriptor_SearchBox", "search box");
 		SlateQueryContext.LastPickedWidget = SearchBox.ToSharedPtr();
 		SlateQueryContext.bIsUIWidget = true;
 	}
@@ -449,8 +483,9 @@ static FText FindItemName(const FWidgetPath& WidgetPath, FAIAssistantSlateQueryC
 			TSharedRef<SMultiLineEditableTextBox> ThisMultiLineEditableTextBoxCast = StaticCastSharedRef<SMultiLineEditableTextBox>(EditableText);
 			HintText = ThisMultiLineEditableTextBoxCast->GetHintText();
 		}
-		OutName = FText::Format(LOCTEXT("ConsoleInputBoxFormat", "{0}"), (HintText.IsEmpty() ? LOCTEXT("ConsoleInputBoxDefaultName","Console") : HintText));
-		OutDescriptor = LOCTEXT("ItemDescriptor_ConsoleInputBox", "input box");
+		SlateQueryContext.ItemName = FText::Format(LOCTEXT("ConsoleInputBoxFormat", "{0}"),
+			(HintText.IsEmpty() ? LOCTEXT("ConsoleInputBoxDefaultName","Console") : HintText));
+		SlateQueryContext.ItemDescriptor = LOCTEXT("ItemDescriptor_ConsoleInputBox", "input box");
 		SlateQueryContext.LastPickedWidget = ConsoleInputBox.ToSharedPtr();
 		SlateQueryContext.bIsUIWidget = true;
 	}
@@ -462,8 +497,8 @@ static FText FindItemName(const FWidgetPath& WidgetPath, FAIAssistantSlateQueryC
 		ChildText = FindChildWidgetWithText(MenuItemBlock);
 		if (!ChildText.IsEmpty())
 		{
-			OutName = ChildText;
-			OutDescriptor = LOCTEXT("ItemDescriptor_Menu", "menu item");
+			SlateQueryContext.ItemName = ChildText;
+			SlateQueryContext.ItemDescriptor = LOCTEXT("ItemDescriptor_Menu", "menu item");
 			SlateQueryContext.LastPickedWidget = MenuItemBlock.ToSharedPtr();
 			SlateQueryContext.bIsUIWidget = true;
 		}
@@ -473,17 +508,21 @@ static FText FindItemName(const FWidgetPath& WidgetPath, FAIAssistantSlateQueryC
 	TSharedRef<SWidget> PropertyRow = FindClosestWidgetOfType(WidgetPath, "SDetailSingleItemRow");
 	if (PropertyRow->GetType() != "SNullWidgetContent")
 	{
+		SlateQueryContext.LastPickedWidget = PropertyRow.ToSharedPtr();
+
 		TSharedRef<SWidget> PropertyNameWidget = FindChildWidgetOfType(PropertyRow, "SPropertyNameWidget");
 		if (PropertyNameWidget->GetType() != "SNullWidgetContent")
 		{
 			ChildText = FindChildWidgetWithText(PropertyNameWidget);
 			if (!ChildText.IsEmpty())
 			{
-				OutName = ChildText;
-				OutDescriptor = LOCTEXT("ItemDescriptor_Property", "setting");
-				SlateQueryContext.LastPickedWidget = PropertyRow.ToSharedPtr();
+				SlateQueryContext.ItemName = ChildText;
+				SlateQueryContext.ItemDescriptor = LOCTEXT("ItemDescriptor_Property", "setting");
 			}
 		}
+
+		// TODO: Some rows e.g. Transform have no SPropertyNameWidget.
+		// TODO: find section name from sibling SDetailCategoryTableRow.
 	}
 
 	// Is it a breadcrumb trail button?
@@ -494,8 +533,8 @@ static FText FindItemName(const FWidgetPath& WidgetPath, FAIAssistantSlateQueryC
 		ChildText = FindChildWidgetWithText(BreadcrumbButton);
 		if (!ChildText.IsEmpty())
 		{
-			OutName = ChildText;
-			OutDescriptor = LOCTEXT("ItemDescriptor_Breadcrumb", "navigation breadcrumb");
+			SlateQueryContext.ItemName = ChildText;
+			SlateQueryContext.ItemDescriptor = LOCTEXT("ItemDescriptor_Breadcrumb", "navigation breadcrumb");
 			SlateQueryContext.LastPickedWidget = BreadcrumbButton.ToSharedPtr();
 			SlateQueryContext.bIsUIWidget = true;
 		}
@@ -505,7 +544,7 @@ static FText FindItemName(const FWidgetPath& WidgetPath, FAIAssistantSlateQueryC
 	TSharedRef<SWidget> AssetTileItem = FindClosestWidgetOfType(WidgetPath, "SAssetTileItem");
 	if (AssetTileItem->GetType() != "SNullWidgetContent")
 	{
-		TSharedRef<SWidget> AssetThumbnail = FindClosestWidgetOfType(WidgetPath, "SAssetThumbnail");
+		TSharedRef<SWidget> AssetThumbnail = FindChildWidgetOfType(AssetTileItem, "SAssetThumbnail");
 		if (AssetThumbnail->GetType() != "SNullWidgetContent")
 		{
 			ChildText = FindChildWidgetWithText(AssetThumbnail); // This returns TYPE of asset instead of name
@@ -513,13 +552,13 @@ static FText FindItemName(const FWidgetPath& WidgetPath, FAIAssistantSlateQueryC
 			{
 				FFormatNamedArguments Args;
 				Args.Add(TEXT("AssetType"), ChildText);
-				OutDescriptor = FText::Format(LOCTEXT("ItemDescriptor_Asset", "{AssetType} asset"), Args);
+				SlateQueryContext.ItemDescriptor = FText::Format(LOCTEXT("ItemDescriptor_Asset", "{AssetType} asset"), Args);
 			}
 
 			if (!SlateQueryContext.CurrentToolTipText.IsEmpty())
 			{
 				// We can use the tooltip for the asset's name.
-				OutName = SlateQueryContext.CurrentToolTipText;
+				SlateQueryContext.ItemName = SlateQueryContext.CurrentToolTipText;
 			}
 			SlateQueryContext.LastPickedWidget = AssetThumbnail.ToSharedPtr();
 			SlateQueryContext.bIsObject = true;
@@ -529,42 +568,92 @@ static FText FindItemName(const FWidgetPath& WidgetPath, FAIAssistantSlateQueryC
 			ChildText = FindChildWidgetWithText(AssetTileItem);
 			if (!ChildText.IsEmpty())
 			{
-				OutName = ChildText;
-				OutDescriptor = LOCTEXT("ItemDescriptor_Folder", "asset folder");
+				SlateQueryContext.ItemName = ChildText;
+				SlateQueryContext.ItemDescriptor = LOCTEXT("ItemDescriptor_Folder", "asset folder");
 				SlateQueryContext.LastPickedWidget = AssetTileItem.ToSharedPtr();
 				SlateQueryContext.bIsObject = true;
 			}
 		}
 	}
 
+	// Is it an asset tree item?
+	TSharedRef<SWidget> AssetTreeItem = FindClosestWidgetOfType(WidgetPath, "SAssetTreeItem");
+	if (AssetTreeItem->GetType() != "SNullWidgetContent")
+	{
+		ChildText = FindChildWidgetWithText(AssetTreeItem);
+		if (!ChildText.IsEmpty())
+		{
+			SlateQueryContext.ItemName = ChildText;
+			SlateQueryContext.ItemDescriptor = LOCTEXT("ItemDescriptor_Folder", "asset folder");
+			SlateQueryContext.LastPickedWidget = AssetTreeItem.ToSharedPtr();
+			SlateQueryContext.bIsObject = true;
+		}
+	}
+
 	// If in Outliner, treat unnamed widgets as actor instances.
 	TSharedRef<SWidget> Outliner = FindClosestWidgetOfType(WidgetPath, "SSceneOutliner");
 	if (Outliner->GetType() != "SNullWidgetContent")
 	{
 		SlateQueryContext.bInOutliner = true;
+		TSharedRef<SWidget> ActorTreeLabel = FindClosestWidgetOfType(WidgetPath, "SActorTreeLabel");
+		if (ActorTreeLabel->GetType() != "SNullWidgetContent")
+		{
+			ChildText = FindChildWidgetWithText(ActorTreeLabel);
+			if (!ChildText.IsEmpty())
+			{
+				SlateQueryContext.ItemName = ChildText;
+				SlateQueryContext.ItemDescriptor = LOCTEXT("ItemDescriptor_OutlinerActor", "actor");
+				SlateQueryContext.LastPickedWidget = ActorTreeLabel.ToSharedPtr();
+				SlateQueryContext.bIsObject = true;
+			}
+		}
+		else
+		{
+			TSharedRef<SWidget> OutlinerRow = FindClosestWidgetOfType(WidgetPath, "SSceneOutlinerTreeRow");
+			if (OutlinerRow->GetType() != "SNullWidgetContent")
+			{
+				SlateQueryContext.LastPickedWidget = OutlinerRow.ToSharedPtr();
+				SlateQueryContext.bIsObject = true;
+				ActorTreeLabel = FindChildWidgetOfType(OutlinerRow, "SActorTreeLabel");
+				if (ActorTreeLabel->GetType() != "SNullWidgetContent")
+				{
+					ChildText = FindChildWidgetWithText(ActorTreeLabel);
+					if (!ChildText.IsEmpty())
+					{
+						SlateQueryContext.ItemName = ChildText;
+						SlateQueryContext.ItemDescriptor = LOCTEXT("ItemDescriptor_OutlinerActor", "actor");
+					}
+				}
+			}
+		}
 	}
 	// If in SubobjectInstanceEditor, treat unnamed widgets as actor instances.
 	TSharedRef<SWidget> SubobjectInstanceEditor = FindClosestWidgetOfType(WidgetPath, "SSubobjectInstanceEditor");
 	if (SubobjectInstanceEditor->GetType() != "SNullWidgetContent")
 	{
 		SlateQueryContext.bInOutliner = true;
+		TSharedRef<SWidget> SubObjectRow = FindClosestWidgetOfType(WidgetPath, "SSubobject_RowWidget");
+		if (SubObjectRow->GetType() != "SNullWidgetContent")
+		{
+			SlateQueryContext.LastPickedWidget = SubObjectRow.ToSharedPtr();
+			SlateQueryContext.bIsObject = true;
+			TSharedRef<SWidget> NameWidget = FindChildWidgetOfType(SubObjectRow, "SInlineEditableTextBlock");
+			if (NameWidget->GetType() != "SNullWidgetContent")
+			{
+				ChildText = FindChildWidgetWithText(NameWidget);
+				SlateQueryContext.ItemName = ChildText;
+				SlateQueryContext.ItemDescriptor = LOCTEXT("ItemDescriptor_SubobjectInstance", "instance");
+			}
+		}
 	}
 
-	if (!OutName.IsEmpty())
-	{
-		FFormatNamedArguments Args;
-		Args.Add(TEXT("ItemName"), OutName);
-		Args.Add(TEXT("ItemDescriptor"), OutDescriptor);
-		return FText::Format(LOCTEXT("ItemName_Format", " the \"{ItemName}\" {ItemDescriptor}"), Args);
-	}
-
-	return OutName;
+	return true;
 }
 
 
-static FText FindTabName(const FWidgetPath& WidgetPath, FAIAssistantSlateQueryContext& SlateQueryContext)
+static bool FindTabName(const FWidgetPath& WidgetPath, FAIAssistantSlateQueryContext& SlateQueryContext)
 {
-	FText OutName = FText();
+	bool bFoundTabName = false;
 
 	TSharedRef<SWidget> TabStack = FindClosestWidgetOfType(WidgetPath, "SDockingTabStack");
 	if (TabStack->GetType() != "SNullWidgetContent")
@@ -576,67 +665,55 @@ static FText FindTabName(const FWidgetPath& WidgetPath, FAIAssistantSlateQueryCo
 			TSharedRef<SDockTab> ThisTabCast = StaticCastSharedRef<SDockTab>(ThisTab);
 			if (ThisTabCast->GetTabRole() == ETabRole::MajorTab || ThisTabCast->GetTabRole() == ETabRole::DocumentTab)
 			{
-				return OutName;
+				break;
 			}
 			if (ThisTabCast->IsForeground())
 			{
 				SlateQueryContext.LastPickedWidget = TabStack.ToSharedPtr();
-				OutName = ThisTabCast->GetTabLabel();
+				SlateQueryContext.StructuredContext.TabName = ThisTabCast->GetTabLabel().ToString();
+				FFormatNamedArguments Args;
+				Args.Add(TEXT("PanelName"), ThisTabCast->GetTabLabel());
+				SlateQueryContext.FormattedTabName = FText::Format(LOCTEXT("TabName_Format", " the {PanelName} panel"), Args);
+				bFoundTabName = true;
 			}
 		}
 	}
 
-	if (!OutName.IsEmpty())
-	{
-		FFormatNamedArguments Args;
-		Args.Add(TEXT("PanelName"), OutName);
-		return FText::Format(LOCTEXT("TabName_Format", " the {PanelName} panel"), Args);
-	}
-	
-	return OutName;
+	return bFoundTabName;
 }
 
 
-static FText GenerateDetailsViewContext(const FWidgetPath& WidgetPath, FAIAssistantSlateQueryContext& SlateQueryContext)
+static void GenerateDetailsViewContext(const FWidgetPath& WidgetPath, FAIAssistantSlateQueryContext& SlateQueryContext)
 {
-	FText OutContext = FText();
-
 	// if we're in a details panel, include the class of actor or object that we're editing
 	const TSharedRef<SWidget> DetailsView = FindClosestWidgetOfType(WidgetPath, "SDetailsView");
 	if (DetailsView->GetType() != "SNullWidgetContent")
 	{
 		const TSharedRef<IDetailsView> DetailsViewCast = StaticCastSharedRef<IDetailsView>(DetailsView);
-		FText QueryContextDetailsObjectAddition =
-			LOCTEXT("QueryContextDetailsObjectAddition",
-					"The property panel I'm working in is editing at least one actor of the class \"{ClassName}\".");
-		FFormatNamedArguments ClassNameArgs;
 		if (DetailsViewCast->GetSelectedActors().Num() > 0)
 		{
 			FString ActorName;
 			DetailsViewCast->GetSelectedActors()[0]->GetClass()->GetName(ActorName);
-			ClassNameArgs.Add(TEXT("ClassName"), FText::FromString(ActorName));
-			SlateQueryContext.GeneratedContextItems.Add(FText::Format(QueryContextDetailsObjectAddition, ClassNameArgs));
+			SlateQueryContext.StructuredContext.ActorClass = ActorName;
 		}
 		else if (DetailsViewCast->GetSelectedObjects().Num() > 0)
 		{
 			FString ObjectName;
 			DetailsViewCast->GetSelectedObjects()[0]->GetClass()->GetName(ObjectName);
-			ClassNameArgs.Add(TEXT("ClassName"), FText::FromString(ObjectName));
-			SlateQueryContext.GeneratedContextItems.Add(FText::Format(QueryContextDetailsObjectAddition, ClassNameArgs));
+			SlateQueryContext.StructuredContext.ActorClass = ObjectName;
 		}
 	}
-
-	return OutContext;
 }
 
 
-static FText FindEditorName(const FWidgetPath& WidgetPath, FAIAssistantSlateQueryContext& SlateQueryContext)
+static bool FindEditorName(const FWidgetPath& WidgetPath, FAIAssistantSlateQueryContext& SlateQueryContext)
 {
 	FText OutName = FText();
 
 	if (!WidgetPath.IsValid())
 	{
-		return OutName;
+		SlateQueryContext.FormattedWindowName = OutName;
+		return false;
 	}
 
 	// is it in the status bar?
@@ -644,7 +721,9 @@ static FText FindEditorName(const FWidgetPath& WidgetPath, FAIAssistantSlateQuer
 	if (StatusBar->GetType() != "SNullWidgetContent")
 	{
 		SlateQueryContext.LastPickedWidget = StatusBar.ToSharedPtr();
-		return LOCTEXT("TabName_StatusBar", " the Status Bar");
+		SlateQueryContext.StructuredContext.TabName = TEXT("StatusBar");
+		SlateQueryContext.FormattedWindowName = LOCTEXT("TabName_StatusBar", " the Status Bar");
+		return true;
 	}
 
 	TSet<FEditorModeID> BlockedModeIDs = TSet<FEditorModeID>({
@@ -662,6 +741,7 @@ static FText FindEditorName(const FWidgetPath& WidgetPath, FAIAssistantSlateQuer
 			TSharedRef<SWidget> LevelEditor = FindChildWidgetOfType(TabStack, "SLevelEditor");
 			if (LevelEditor->GetType() != "SNullWidgetContent")
 			{
+				SlateQueryContext.StructuredContext.EditorName = TEXT("LevelEditor");
 				FText ModeString;
 				ULevelEditorSubsystem* LevelEditorSubsystem = GEditor->GetEditorSubsystem<ULevelEditorSubsystem>();
 				if (FEditorModeTools* ModeManager = LevelEditorSubsystem->GetLevelEditorModeManager())
@@ -675,19 +755,22 @@ static FText FindEditorName(const FWidgetPath& WidgetPath, FAIAssistantSlateQuer
 							FFormatNamedArguments Args;
 							Args.Add(TEXT("ModeName"), ModeInfo.Name);
 							ModeString = FText::Format(LOCTEXT("EditorName_ModeFormat", "with {ModeName} mode active"), Args);
+							SlateQueryContext.StructuredContext.EditorMode = ModeInfo.Name.ToString();
 						}
 					}
 				}
 				SlateQueryContext.LastPickedWidget = LevelEditor.ToSharedPtr();
 				if (ModeString.IsEmpty())
 				{
-					return LOCTEXT("EditorName_LevelEditor", " the Level Editor");
+					SlateQueryContext.FormattedWindowName = LOCTEXT("EditorName_LevelEditor", " the Level Editor");
+					return true;
 				}
 				else
 				{
 					FFormatNamedArguments Args;
 					Args.Add(TEXT("WithModeActive"), ModeString);
-					return FText::Format(LOCTEXT("EditorName_LevelEditorAndMode", " the Level Editor {WithModeActive}"), Args);
+					SlateQueryContext.FormattedWindowName = FText::Format(LOCTEXT("EditorName_LevelEditorAndMode", " the Level Editor {WithModeActive}"), Args);
+					return true;
 				}
 			}
 
@@ -708,12 +791,14 @@ static FText FindEditorName(const FWidgetPath& WidgetPath, FAIAssistantSlateQuer
 						for (auto& ThisEditedAsset : AllEditedAssets)
 						{
 							IAssetEditorInstance* ThisEditor = AssetEditorSubsystem->FindEditorForAsset(ThisEditedAsset, false);
-							if (ThisEditor->GetAssociatedTabManager() == PickedWidgetTabManager)
+							if (ThisEditor && ThisEditor->GetAssociatedTabManager() == PickedWidgetTabManager)
 							{
+								SlateQueryContext.StructuredContext.EditorName = ThisEditor->GetEditorName().ToString();
 								SlateQueryContext.LastPickedWidget = TabStack.ToSharedPtr();
 								FFormatNamedArguments Args;
 								Args.Add(TEXT("AssetEditorName"), FText::FromName(ThisEditor->GetEditorName()));
-								return FText::Format(LOCTEXT("EditorName_AssetEditorFormat", " the {AssetEditorName} editor"), Args);
+								SlateQueryContext.FormattedWindowName = FText::Format(LOCTEXT("EditorName_AssetEditorFormat", " the {AssetEditorName} editor"), Args);
+								return true;
 							}
 						}
 					}
@@ -745,10 +830,13 @@ static FText FindEditorName(const FWidgetPath& WidgetPath, FAIAssistantSlateQuer
 
 		if (!DrawerName.IsEmpty())
 		{
+			SlateQueryContext.StructuredContext.WindowName = DrawerOverlay->GetType().ToString();
+			SlateQueryContext.StructuredContext.TabName = DrawerName.ToString();
 			SlateQueryContext.LastPickedWidget = DrawerWidget.ToSharedPtr();
 			FFormatNamedArguments Args;
 			Args.Add(TEXT("DrawerName"), DrawerName);
-			return FText::Format(LOCTEXT("DrawerNameFormat", "the {DrawerName} drawer"), Args);
+			SlateQueryContext.FormattedWindowName = FText::Format(LOCTEXT("DrawerNameFormat", "the {DrawerName} drawer"), Args);
+			return true;
 		}
 	}
 
@@ -759,7 +847,9 @@ static FText FindEditorName(const FWidgetPath& WidgetPath, FAIAssistantSlateQuer
 		IMainFrameModule& MainFrameModule = FModuleManager::LoadModuleChecked<IMainFrameModule>(TEXT("MainFrame"));
 		if (WidgetWindow.ToSharedPtr() == MainFrameModule.GetParentWindow())
 		{
-			return LOCTEXT("WindowMainFrame_Title", " the editor's main window");
+			SlateQueryContext.StructuredContext.WindowName = TEXT("MainFrame");
+			SlateQueryContext.FormattedWindowName = LOCTEXT("WindowMainFrame_Title", " the editor's main window");
+			return true;
 		}
 	}
 	
@@ -767,18 +857,20 @@ static FText FindEditorName(const FWidgetPath& WidgetPath, FAIAssistantSlateQuer
 	FText WindowTitle = WidgetWindow->GetTitle();
 	if (!WindowTitle.IsEmpty())
 	{
+		SlateQueryContext.StructuredContext.WindowName = WindowTitle.ToString();
 		SlateQueryContext.LastPickedWidget = WidgetWindow.ToSharedPtr();
 		FFormatNamedArguments Args;
 		Args.Add(TEXT("WindowTitle"), WindowTitle);
-		return FText::Format(LOCTEXT("WindowNameFormat_Title", " the {WindowTitle} window"), Args);
+		SlateQueryContext.FormattedWindowName = FText::Format(LOCTEXT("WindowNameFormat_Title", " the {WindowTitle} window"), Args);
+		return true;
 	}
 
 	// TODO: how to find other editor contexts?
-	return OutName;
+	return false;
 }
 
 
-static FText FindTextUnderCursor(const FWidgetPath& WidgetPath, FAIAssistantSlateQueryContext& AccumulatedSlateQueryContext)
+static void FindTextUnderCursor(const FWidgetPath& WidgetPath, FAIAssistantSlateQueryContext& AccumulatedSlateQueryContext)
 {
 	FText OutText;
 
@@ -809,10 +901,10 @@ static FText FindTextUnderCursor(const FWidgetPath& WidgetPath, FAIAssistantSlat
 		if (!OutText.IsEmpty())
 		{
 			AccumulatedSlateQueryContext.LastPickedWidget = WidgetToTest;
+			AccumulatedSlateQueryContext.TextUnderCursor = OutText;
+			AccumulatedSlateQueryContext.StructuredContext.TextUnderCursor = OutText.ToString();
 		}
 	}
-
-	return OutText;
 }
 
 
@@ -864,35 +956,37 @@ void UE::AIAssistant::SlateQuerier::QueryAIAssistantAboutSlateWidgetUnderCursor(
 	const FWidgetPath ContextWidgetPath = FindContextWidgetPath(WidgetPath);
 
 	// Is there an identifiable asset editor, mode, or window?
-	const FText WindowName = FindEditorName(ContextWidgetPath, SlateQueryContext);
+	FindEditorName(ContextWidgetPath, SlateQueryContext);
 
 	// Is there an identifiable tab?
 	// Look up the chain for an SDockingTabStack, then look down to find an SDockTab.
-	const FText TabName = FindTabName(ContextWidgetPath, SlateQueryContext);
+	FindTabName(ContextWidgetPath, SlateQueryContext);
 
 	// Is there text under the cursor??
-	const FText TextUnderCursor = FindTextUnderCursor(WidgetPath, SlateQueryContext);
+	FindTextUnderCursor(WidgetPath, SlateQueryContext);
 
 	// Is there an identifiable item?
-	const FText ItemName = FindItemName(WidgetPath, SlateQueryContext);
+	FindItemName(WidgetPath, SlateQueryContext);
 
-	// Special case for actor listings in Outliner.
-	if (ItemName.IsEmpty() && SlateQueryContext.bInOutliner && !SlateQueryContext.bIsUIWidget)
+	// Construct query.
 	{
-		SlateQueryContext.bIsObject = true;
-	}
+		FText FormattedItemName = SlateQueryContext.ItemName;
+		if (!FormattedItemName.IsEmpty())
+		{
+			FFormatNamedArguments Args;
+			Args.Add(TEXT("ItemName"), SlateQueryContext.ItemName);
+			Args.Add(TEXT("ItemDescriptor"), SlateQueryContext.ItemDescriptor);
+			FormattedItemName = FText::Format(LOCTEXT("ItemName_Format", " the \"{ItemName}\" {ItemDescriptor}"), Args);
+		}
 
-	{
 		FText NameToQuery;
-		bool bHasName = false;
-		if (ItemName.IsEmpty() && !TextUnderCursor.IsEmpty())
+		if (FormattedItemName.IsEmpty() && !SlateQueryContext.TextUnderCursor.IsEmpty())
 		{
-			NameToQuery = TextUnderCursor;
+			NameToQuery = SlateQueryContext.TextUnderCursor;
 		}
-		else if (!ItemName.IsEmpty() || !TabName.IsEmpty() || !WindowName.IsEmpty())
+		else if (!FormattedItemName.IsEmpty() || !SlateQueryContext.FormattedTabName.IsEmpty() || !SlateQueryContext.FormattedWindowName.IsEmpty())
 		{
-			NameToQuery = (!ItemName.IsEmpty()) ? ItemName : (!TabName.IsEmpty()) ? TabName : WindowName;
-			bHasName = true;
+			NameToQuery = (!FormattedItemName.IsEmpty()) ? FormattedItemName : (!SlateQueryContext.FormattedTabName.IsEmpty()) ? SlateQueryContext.FormattedTabName : SlateQueryContext.FormattedWindowName;
 		}
 
 		static const FText QueryFormatWidget = LOCTEXT("QueryFormat1", "I would like to know what {NameToQuery} does.");
@@ -909,9 +1003,9 @@ void UE::AIAssistant::SlateQuerier::QueryAIAssistantAboutSlateWidgetUnderCursor(
 		static const FText QueryAdditionalInstructionsGeneral = LOCTEXT("QueryAdditionalInstructions3",
 			"If the object is clickable, like a button, then explain what happens if I click on it. If the object is something you select, like an asset or a blueprint node, describe what the object is or what it does.");
 		static const FText StyleInstructionsWidget = LOCTEXT("StyleInstructions1",
-			"Use bold text to emphasize Unreal Engine specific terms. Start with a quick overview paragraph, then add key points in one or two formatted sections with headers and bullet points. Don't include a summary at the end.");
+			"Use bold text to emphasize Unreal Engine specific terms. Start with a quick overview paragraph, then add key points limited to one or two formatted sections with headers and bullet points. Don't include a summary at the end.");
 		static const FText StyleInstructionsObject = LOCTEXT("StyleInstructions2",
-			"Use bold text to emphasize Unreal Engine specific terms. Start with a quick overview paragraph, then add key points in one or two formatted sections with headers and bullet points. Don't include a summary at the end.");
+			"Use bold text to emphasize Unreal Engine specific terms. Start with a quick overview paragraph, then add key points limited to one or two formatted sections with headers and bullet points. Don't include a summary at the end.");
 
 		if (!NameToQuery.IsEmpty())
 		{
@@ -949,53 +1043,51 @@ void UE::AIAssistant::SlateQuerier::QueryAIAssistantAboutSlateWidgetUnderCursor(
 		}
 	}
 
+	// Construct structured context.
 	{
-		if (!TabName.IsEmpty() || !WindowName.IsEmpty())
-		{
-			FFormatNamedArguments Args;
-			if (!TabName.IsEmpty() && !WindowName.IsEmpty())
-			{
-				Args.Add(TEXT("TabName"), TabName);
-				Args.Add(TEXT("WindowName"), WindowName);
-				SlateQueryContext.GeneratedContextItems.Add(FText::Format(LOCTEXT("QueryContextTabWindowLocation", "I am working in {TabName} of {WindowName}."), Args));
-			}
-			else
-			{
-				Args.Add(TEXT("MaybeEmptyTabName"), TabName);
-				Args.Add(TEXT("MaybeEmptyWindowName"), WindowName);
-				SlateQueryContext.GeneratedContextItems.Add(FText::Format(LOCTEXT("QueryContextTabOrWindowLocation", "I am working in {MaybeEmptyTabName}{MaybeEmptyWindowName}."), Args));
-			}
+		SlateQueryContext.StructuredContext.ToolTipText = SlateQueryContext.CurrentToolTipText.ToString();
 
-			if (const FText ModeToolContext = GenerateModeToolContext(WidgetPath);
-				!ModeToolContext.IsEmpty())
-			{
-				SlateQueryContext.GeneratedContextItems.Add(ModeToolContext);
-			}
+		GenerateModeToolContext(WidgetPath, SlateQueryContext);
+		GenerateDetailsViewContext(WidgetPath, SlateQueryContext);
+
+		if (!SlateQueryContext.ItemName.IsEmpty())
+		{
+			SlateQueryContext.StructuredContext.ItemName = SlateQueryContext.ItemName.ToString();
+			SlateQueryContext.StructuredContext.ItemDescriptor = SlateQueryContext.ItemDescriptor.ToString();
 		}
 
-		if (const FText DetailsViewContext = GenerateDetailsViewContext(WidgetPath, SlateQueryContext);
-			!DetailsViewContext.IsEmpty())
+		// Get top level window title.
+		TSharedRef<SWindow> WidgetWindow = WidgetPath.GetWindow();
+		FText WindowTitle = WidgetWindow->GetTitle();
+		if (!WindowTitle.IsEmpty())
 		{
-			SlateQueryContext.GeneratedContextItems.Add(DetailsViewContext);
+			SlateQueryContext.StructuredContext.TopWindowName = WindowTitle.ToString();
 		}
 
-		if (!SlateQueryContext.CurrentToolTipText.IsEmpty())
+		// Get picked widget type.
+		if (SlateQueryContext.LastPickedWidget.IsValid())
 		{
-			FText ToolTipContextFormat = LOCTEXT("QueryContextToolTipAddition", "The text of the current tooltip is: \"{ToolTipText}\", but you don't need to mention it.");
-			FFormatNamedArguments Args;
-			Args.Add(TEXT("ToolTipText"), SlateQueryContext.CurrentToolTipText);
-			FText ToolTipContext = FText::Format(ToolTipContextFormat, Args);
-			SlateQueryContext.GeneratedContextItems.Add(ToolTipContext);
+			SlateQueryContext.StructuredContext.WidgetType = SlateQueryContext.LastPickedWidget->GetType().ToString();
 		}
 
-		if (!SlateQueryContext.GeneratedContextItems.IsEmpty())
+		if (SlateQueryContext.ItemName.IsEmpty())
 		{
-			for (const FText& ContextItem : SlateQueryContext.GeneratedContextItems)
+			SlateQueryContext.StructuredContext.WidgetPath = TArray<FString>();
+			// Add simplified widget path
+			for (int32 WidgetIndex = 0; WidgetIndex < WidgetPath.Widgets.Num(); WidgetIndex++)
 			{
-				FFormatNamedArguments Args;
-				Args.Add(TEXT("GeneratedContextSoFar"), SlateQueryContext.GeneratedContext);
-				Args.Add(TEXT("ContextItemToAdd"), ContextItem);
-				SlateQueryContext.GeneratedContext = FText::Format(LOCTEXT("LastGeneratedContentFormat", "{GeneratedContextSoFar} {ContextItemToAdd}"), Args);
+				const FArrangedWidget* ArrangedWidget = &WidgetPath.Widgets[WidgetIndex];
+				const TSharedRef<SWidget>& ThisWidget = ArrangedWidget->Widget;
+				FString NextWidgetString = ThisWidget->GetType().ToString();
+				// Skip uninteresting widgets
+				if (NextWidgetString != TEXT("SBorder") && NextWidgetString != TEXT("SBox") &&
+					NextWidgetString != TEXT("SScrollBorder") &&
+					NextWidgetString != TEXT("SOverlay") && NextWidgetString != TEXT("SVerticalBox") &&
+					NextWidgetString != TEXT("SHorizontalBox") && NextWidgetString != TEXT("SSplitter") &&
+					NextWidgetString != TEXT("SDockingSplitter"))
+				{
+					SlateQueryContext.StructuredContext.WidgetPath->Add(NextWidgetString);
+				}
 			}
 		}
 	}
@@ -1003,21 +1095,19 @@ void UE::AIAssistant::SlateQuerier::QueryAIAssistantAboutSlateWidgetUnderCursor(
 	// OPTIONAL - We're done using these members. We can clear them to reduce size as they're
 	// not used below.
 	SlateQueryContext.LastPickedWidget.Reset();
-	SlateQueryContext.GeneratedContextItems.Empty();
 
 	// Send widget query to AI Assistant.
 	TSharedPtr<SAIAssistantWebBrowser> WebBrowser =
 		UAIAssistantSubsystem::GetAIAssistantWebBrowserWidget();
 	WebBrowser->CreateConversation();
-				
+
 	const FString VisiblePromptString =
 		CleanExtraWhiteSpaceFromString(SlateQueryContext.GeneratedQuery.ToString());
 
 	const FString HiddenContextString = 
 		CleanExtraWhiteSpaceFromString(SlateQueryContext.GeneratedQueryInstructions.ToString()) +
-		FText(LOCTEXT("ContextPrefix", "(Context: ")).ToString() +
-		CleanExtraWhiteSpaceFromString(SlateQueryContext.GeneratedContext.ToString()) +
-		FText(LOCTEXT("ContextPostfix", ")")).ToString();
+		TEXT(" ContextData=") +
+		SlateQueryContext.StructuredContext.ToJson(false);
 
 	WebBrowser->AddUserMessageToConversation(VisiblePromptString, HiddenContextString);
 }
