Set a new place in the Window menu to launch the assistant.

diff --git a/Engine/Plugins/Experimental/AIAssistant/Source/AIAssistant/AIAssistant.Build.cs b/Engine/Plugins/Experimental/AIAssistant/Source/AIAssistant/AIAssistant.Build.cs
--- a/Engine/Plugins/Experimental/AIAssistant/Source/AIAssistant/AIAssistant.Build.cs
+++ b/Engine/Plugins/Experimental/AIAssistant/Source/AIAssistant/AIAssistant.Build.cs
@@ -68,7 +68,8 @@ public class AIAssistant : ModuleRules
 				// For clipboard actions -
 				"ApplicationCore",
 				// For web browser -
-				"WebBrowser"
+				"WebBrowser",
+				"WorkspaceMenuStructure"
 			}
 			);
 		
diff --git a/Engine/Plugins/Experimental/AIAssistant/Source/AIAssistant/Private/AIAssistant.cpp b/Engine/Plugins/Experimental/AIAssistant/Source/AIAssistant/Private/AIAssistant.cpp
--- a/Engine/Plugins/Experimental/AIAssistant/Source/AIAssistant/Private/AIAssistant.cpp
+++ b/Engine/Plugins/Experimental/AIAssistant/Source/AIAssistant/Private/AIAssistant.cpp
@@ -14,6 +14,9 @@
 #include "AIAssistantStyle.h"
 #include "AIAssistantWebBrowser.h"
 
+#include "WorkspaceMenuStructure.h"
+#include "WorkspaceMenuStructureModule.h"
+
 #define LOCTEXT_NAMESPACE "FAIAssistantModule"
 
 
@@ -52,13 +55,6 @@ static const FName AIAssistantTabName("AIAssistant");
 	
 	
 	UToolMenus::RegisterStartupCallback(FSimpleMulticastDelegate::FDelegate::CreateRaw(this, &FAIAssistantModule::RegisterMenus));
-
-	
-	FGlobalTabmanager::Get()->RegisterNomadTabSpawner(AIAssistantTabName, FOnSpawnTab::CreateRaw(this, &FAIAssistantModule::OnSpawnPluginTab))
-		.SetDisplayName(LOCTEXT("FAIAssistantTabTitle", "AI Assistant"))
-		.SetMenuType(ETabSpawnerMenuType::Hidden)
-		.SetIcon(FSlateIcon("AIAssistantStyle","AIAssistant.OpenPluginWindow"));
-
 	
 	InputProcessor = MakeShared<FAIAssistantInputProcessor>(PluginCommands);
 
@@ -163,23 +159,13 @@ void FAIAssistantModule::RegisterMenus()
 	FToolMenuOwnerScoped OwnerScoped(this);
 
 	{
-		UToolMenu* Menu = UToolMenus::Get()->ExtendMenu("LevelEditor.MainMenu.Window");
-
-		FToolMenuSection& Section = Menu->FindOrAddSection("WindowLayout");
-
-		Section.AddMenuEntryWithCommandList(FAIAssistantCommands::Get().OpenAIAssistantTab, PluginCommands);
+		TSharedPtr<FWorkspaceItem> AssistantWorkspaceCategory(WorkspaceMenu::GetMenuStructure().GetStructureRoot()->AddGroup("Assistance", LOCTEXT("FAIAssistantMenuSectionTitle", "Assistance"), FSlateIcon()));
+		FTabSpawnerEntry AssistantTabSpawner = FGlobalTabmanager::Get()->RegisterNomadTabSpawner(AIAssistantTabName, FOnSpawnTab::CreateRaw(this, &FAIAssistantModule::OnSpawnPluginTab))
+			.SetDisplayName(LOCTEXT("FAIAssistantTabTitle", "AI Assistant"))
+			.SetIcon(FSlateIcon("AIAssistantStyle", "AIAssistant.OpenPluginWindow"))
+			.SetGroup(AssistantWorkspaceCategory.ToSharedRef());
 	}
 
-	{
-		UToolMenu* ToolbarMenu = UToolMenus::Get()->ExtendMenu("LevelEditor.LevelEditorToolBar");
-
-		FToolMenuSection& Section = ToolbarMenu->FindOrAddSection("Settings");
-
-		FToolMenuEntry& Entry = Section.AddEntry(FToolMenuEntry::InitToolBarButton(
-			FAIAssistantCommands::Get().OpenAIAssistantTab));
-
-		Entry.SetCommandList(PluginCommands);
-	}
 }
 
 
