Back out changes that caused a compile warning on Mac.

#FailedValidations Preflight

diff --git a/Engine/Plugins/Experimental/AIAssistant/Source/AIAssistant/AIAssistant.Build.cs b/Engine/Plugins/Experimental/AIAssistant/Source/AIAssistant/AIAssistant.Build.cs
--- a/Engine/Plugins/Experimental/AIAssistant/Source/AIAssistant/AIAssistant.Build.cs
+++ b/Engine/Plugins/Experimental/AIAssistant/Source/AIAssistant/AIAssistant.Build.cs
@@ -68,8 +68,7 @@ public class AIAssistant : ModuleRules
 				// For clipboard actions -
 				"ApplicationCore",
 				// For web browser -
-				"WebBrowser",
-				"WorkspaceMenuStructure"
+				"WebBrowser"
 			}
 			);
 		
diff --git a/Engine/Plugins/Experimental/AIAssistant/Source/AIAssistant/Private/AIAssistant.cpp b/Engine/Plugins/Experimental/AIAssistant/Source/AIAssistant/Private/AIAssistant.cpp
--- a/Engine/Plugins/Experimental/AIAssistant/Source/AIAssistant/Private/AIAssistant.cpp
+++ b/Engine/Plugins/Experimental/AIAssistant/Source/AIAssistant/Private/AIAssistant.cpp
@@ -14,9 +14,6 @@
 #include "AIAssistantStyle.h"
 #include "AIAssistantWebBrowser.h"
 
-#include "WorkspaceMenuStructure.h"
-#include "WorkspaceMenuStructureModule.h"
-
 #define LOCTEXT_NAMESPACE "FAIAssistantModule"
 
 
@@ -55,6 +52,13 @@ static const FName AIAssistantTabName("AIAssistant");
 	
 	
 	UToolMenus::RegisterStartupCallback(FSimpleMulticastDelegate::FDelegate::CreateRaw(this, &FAIAssistantModule::RegisterMenus));
+
+	
+	FGlobalTabmanager::Get()->RegisterNomadTabSpawner(AIAssistantTabName, FOnSpawnTab::CreateRaw(this, &FAIAssistantModule::OnSpawnPluginTab))
+		.SetDisplayName(LOCTEXT("FAIAssistantTabTitle", "AI Assistant"))
+		.SetMenuType(ETabSpawnerMenuType::Hidden)
+		.SetIcon(FSlateIcon("AIAssistantStyle","AIAssistant.OpenPluginWindow"));
+
 	
 	InputProcessor = MakeShared<FAIAssistantInputProcessor>(PluginCommands);
 
@@ -159,13 +163,23 @@ void FAIAssistantModule::RegisterMenus()
 	FToolMenuOwnerScoped OwnerScoped(this);
 
 	{
-		TSharedPtr<FWorkspaceItem> AssistantWorkspaceCategory(WorkspaceMenu::GetMenuStructure().GetStructureRoot()->AddGroup("Assistance", LOCTEXT("FAIAssistantMenuSectionTitle", "Assistance"), FSlateIcon()));
-		FTabSpawnerEntry AssistantTabSpawner = FGlobalTabmanager::Get()->RegisterNomadTabSpawner(AIAssistantTabName, FOnSpawnTab::CreateRaw(this, &FAIAssistantModule::OnSpawnPluginTab))
-			.SetDisplayName(LOCTEXT("FAIAssistantTabTitle", "AI Assistant"))
-			.SetIcon(FSlateIcon("AIAssistantStyle", "AIAssistant.OpenPluginWindow"))
-			.SetGroup(AssistantWorkspaceCategory.ToSharedRef());
+		UToolMenu* Menu = UToolMenus::Get()->ExtendMenu("LevelEditor.MainMenu.Window");
+
+		FToolMenuSection& Section = Menu->FindOrAddSection("WindowLayout");
+
+		Section.AddMenuEntryWithCommandList(FAIAssistantCommands::Get().OpenAIAssistantTab, PluginCommands);
 	}
 
+	{
+		UToolMenu* ToolbarMenu = UToolMenus::Get()->ExtendMenu("LevelEditor.LevelEditorToolBar");
+
+		FToolMenuSection& Section = ToolbarMenu->FindOrAddSection("Settings");
+
+		FToolMenuEntry& Entry = Section.AddEntry(FToolMenuEntry::InitToolBarButton(
+			FAIAssistantCommands::Get().OpenAIAssistantTab));
+
+		Entry.SetCommandList(PluginCommands);
+	}
 }
 
 
