[AI Assistant] improve UI query widget detection
- improve detection of details tab property settings name
- provide details property category/section name
- fix issue with detecting asset editor in top-level tabs

diff --git a/Engine/Plugins/Experimental/AIAssistant/Source/AIAssistant/Private/AIAssistantSlateQuerier.cpp b/Engine/Plugins/Experimental/AIAssistant/Source/AIAssistant/Private/AIAssistantSlateQuerier.cpp
--- a/Engine/Plugins/Experimental/AIAssistant/Source/AIAssistant/Private/AIAssistantSlateQuerier.cpp
+++ b/Engine/Plugins/Experimental/AIAssistant/Source/AIAssistant/Private/AIAssistantSlateQuerier.cpp
@@ -44,6 +44,7 @@ struct FAIAssistantSlateQueryStructuredContext : public FJsonSerializable
 	TOptional<FString> ActorClass;
 	TOptional<FString> ItemName;
 	TOptional<FString> ItemDescriptor;
+	TOptional<FString> CategoryName;
 	TOptional<TArray<FString>> WidgetPath;
 
 	BEGIN_JSON_SERIALIZER
@@ -59,6 +60,7 @@ struct FAIAssistantSlateQueryStructuredContext : public FJsonSerializable
 		JSON_SERIALIZE_OPTIONAL("actor_class", ActorClass);
 		JSON_SERIALIZE_OPTIONAL("item_name", ItemName);
 		JSON_SERIALIZE_OPTIONAL("item_descriptor", ItemDescriptor);
+		JSON_SERIALIZE_OPTIONAL("category_name", CategoryName);
 		JSON_SERIALIZE_OPTIONAL_ARRAY("widget_path", WidgetPath);
 	END_JSON_SERIALIZER
 };
@@ -82,7 +84,6 @@ struct FAIAssistantSlateQueryContext
 	FText ItemDescriptor = FText();
 	bool bIsUIWidget = false;
 	bool bIsObject = false;
-	bool bInOutliner = false;
 
 	FText GeneratedQuery = FText();
 	FText GeneratedQueryInstructions = FText();
@@ -520,9 +521,45 @@ static bool FindItemName(const FWidgetPath& WidgetPath, FAIAssistantSlateQueryCo
 				SlateQueryContext.ItemDescriptor = LOCTEXT("ItemDescriptor_Property", "setting");
 			}
 		}
+		else
+		{
+			// Some rows e.g. Transform have no SPropertyNameWidget. Best we can do is look for a text label.
+			ChildText = FindChildWidgetWithText(PropertyRow);
+			if (!ChildText.IsEmpty())
+			{
+				SlateQueryContext.ItemName = ChildText;
+				SlateQueryContext.ItemDescriptor = LOCTEXT("ItemDescriptor_Property", "setting");
+			}
+		}
 
-		// TODO: Some rows e.g. Transform have no SPropertyNameWidget.
-		// TODO: find section name from sibling SDetailCategoryTableRow.
+		// Find section name from sibling SDetailCategoryTableRow.
+		if (PropertyRow->IsParentValid())
+		{
+			TSharedPtr<SWidget> ListPanel = PropertyRow->GetParentWidget();
+			TSharedRef<SWidget> PreviousDetailCategory = SNullWidget::NullWidget;
+			for (int32 ChildIdx = 0; ChildIdx < ListPanel->GetChildren()->Num(); ChildIdx++)
+			{
+				TSharedRef<SWidget> ThisWidget = ListPanel->GetChildren()->GetChildAt(ChildIdx);
+				if (ThisWidget == PropertyRow)
+				{
+					break;
+				}
+				if (ThisWidget->GetType() == "SDetailCategoryTableRow")
+				{
+					PreviousDetailCategory = ThisWidget;
+				}
+			}
+			if (PreviousDetailCategory->GetType() != "SNullWidgetContent")
+			{
+				ChildText = FindChildWidgetWithText(PreviousDetailCategory);
+				if (!ChildText.IsEmpty())
+				{
+					SlateQueryContext.StructuredContext.CategoryName = ChildText.ToString();
+					SlateQueryContext.ItemDescriptor = FText::Format(
+						LOCTEXT("ItemDescriptor_PropertyWithCategory", "setting in the {0} category"), ChildText);
+				}
+			}
+		}
 	}
 
 	// Is it a breadcrumb trail button?
@@ -590,11 +627,10 @@ static bool FindItemName(const FWidgetPath& WidgetPath, FAIAssistantSlateQueryCo
 		}
 	}
 
-	// If in Outliner, treat unnamed widgets as actor instances.
+	// Is it an actor in the Outliner?
 	TSharedRef<SWidget> Outliner = FindClosestWidgetOfType(WidgetPath, "SSceneOutliner");
 	if (Outliner->GetType() != "SNullWidgetContent")
 	{
-		SlateQueryContext.bInOutliner = true;
 		TSharedRef<SWidget> ActorTreeLabel = FindClosestWidgetOfType(WidgetPath, "SActorTreeLabel");
 		if (ActorTreeLabel->GetType() != "SNullWidgetContent")
 		{
@@ -627,11 +663,10 @@ static bool FindItemName(const FWidgetPath& WidgetPath, FAIAssistantSlateQueryCo
 			}
 		}
 	}
-	// If in SubobjectInstanceEditor, treat unnamed widgets as actor instances.
+	// Is it an object in the Subobject Instance editor?
 	TSharedRef<SWidget> SubobjectInstanceEditor = FindClosestWidgetOfType(WidgetPath, "SSubobjectInstanceEditor");
 	if (SubobjectInstanceEditor->GetType() != "SNullWidgetContent")
 	{
-		SlateQueryContext.bInOutliner = true;
 		TSharedRef<SWidget> SubObjectRow = FindClosestWidgetOfType(WidgetPath, "SSubobject_RowWidget");
 		if (SubObjectRow->GetType() != "SNullWidgetContent")
 		{
@@ -731,52 +766,54 @@ static bool FindEditorName(const FWidgetPath& WidgetPath, FAIAssistantSlateQuery
 		"EditMode.SubTrackEditMode"
 		});
 
-	TSharedRef<SWidget> DockingArea = FindFirstWidgetOfType(WidgetPath, "SDockingArea");
-	if (DockingArea->GetType() != "SNullWidgetContent")
+	// are we in the level editor?
+	TSharedRef<SWidget> LevelEditor = FindFirstWidgetOfType(WidgetPath, "SLevelEditor");
+	if (LevelEditor->GetType() != "SNullWidgetContent")
 	{
-		TSharedRef<SWidget> TabStack = FindChildWidgetOfType(DockingArea.ToSharedPtr(), "SDockingTabStack");
-		if (TabStack->GetType() != "SNullWidgetContent")
+		SlateQueryContext.StructuredContext.EditorName = TEXT("LevelEditor");
+		FText ModeString;
+		ULevelEditorSubsystem* LevelEditorSubsystem = GEditor->GetEditorSubsystem<ULevelEditorSubsystem>();
+		if (FEditorModeTools* ModeManager = LevelEditorSubsystem->GetLevelEditorModeManager())
 		{
-			// are we in the level editor?
-			TSharedRef<SWidget> LevelEditor = FindChildWidgetOfType(TabStack, "SLevelEditor");
-			if (LevelEditor->GetType() != "SNullWidgetContent")
+			UAssetEditorSubsystem* AssetEditorSubsystem = GEditor->GetEditorSubsystem<UAssetEditorSubsystem>();
+			TArray<FEditorModeInfo> ModeInfos = AssetEditorSubsystem->GetEditorModeInfoOrderedByPriority();
+			for (const FEditorModeInfo& ModeInfo : ModeInfos)
 			{
-				SlateQueryContext.StructuredContext.EditorName = TEXT("LevelEditor");
-				FText ModeString;
-				ULevelEditorSubsystem* LevelEditorSubsystem = GEditor->GetEditorSubsystem<ULevelEditorSubsystem>();
-				if (FEditorModeTools* ModeManager = LevelEditorSubsystem->GetLevelEditorModeManager())
-				{
-					UAssetEditorSubsystem* AssetEditorSubsystem = GEditor->GetEditorSubsystem<UAssetEditorSubsystem>();
-					TArray<FEditorModeInfo> ModeInfos = AssetEditorSubsystem->GetEditorModeInfoOrderedByPriority();
-					for (const FEditorModeInfo& ModeInfo : ModeInfos)
-					{
-						if (ModeString.IsEmpty() && ModeManager->IsModeActive(ModeInfo.ID) && !BlockedModeIDs.Contains(ModeInfo.ID))
-						{
-							FFormatNamedArguments Args;
-							Args.Add(TEXT("ModeName"), ModeInfo.Name);
-							ModeString = FText::Format(LOCTEXT("EditorName_ModeFormat", "with {ModeName} mode active"), Args);
-							SlateQueryContext.StructuredContext.EditorMode = ModeInfo.Name.ToString();
-						}
-					}
-				}
-				SlateQueryContext.LastPickedWidget = LevelEditor.ToSharedPtr();
-				if (ModeString.IsEmpty())
-				{
-					SlateQueryContext.FormattedWindowName = LOCTEXT("EditorName_LevelEditor", " the Level Editor");
-					return true;
-				}
-				else
+				if (ModeString.IsEmpty() && ModeManager->IsModeActive(ModeInfo.ID) && !BlockedModeIDs.Contains(ModeInfo.ID))
 				{
 					FFormatNamedArguments Args;
-					Args.Add(TEXT("WithModeActive"), ModeString);
-					SlateQueryContext.FormattedWindowName = FText::Format(LOCTEXT("EditorName_LevelEditorAndMode", " the Level Editor {WithModeActive}"), Args);
-					return true;
+					Args.Add(TEXT("ModeName"), ModeInfo.Name);
+					ModeString = FText::Format(LOCTEXT("EditorName_ModeFormat", "with {ModeName} mode active"), Args);
+					SlateQueryContext.StructuredContext.EditorMode = ModeInfo.Name.ToString();
 				}
 			}
+		}
+		SlateQueryContext.LastPickedWidget = LevelEditor.ToSharedPtr();
+		if (ModeString.IsEmpty())
+		{
+			SlateQueryContext.FormattedWindowName = LOCTEXT("EditorName_LevelEditor", " the Level Editor");
+			return true;
+		}
+		else
+		{
+			FFormatNamedArguments Args;
+			Args.Add(TEXT("WithModeActive"), ModeString);
+			SlateQueryContext.FormattedWindowName = FText::Format(LOCTEXT("EditorName_LevelEditorAndMode", " the Level Editor {WithModeActive}"), Args);
+			return true;
+		}
+	}
 
-			// if not, are we in a different asset editor?
+	// if not, are we in a different asset editor?
+	TSharedRef<SWidget> StandaloneAssetEditorHost = FindFirstWidgetOfType(WidgetPath, "SStandaloneAssetEditorToolkitHost");
+	if (StandaloneAssetEditorHost->GetType() != "SNullWidgetContent")
+	{
+		TArray<TSharedRef<SWidget>> TabStacks;
+		FindChildWidgetsOfType(TabStacks, StandaloneAssetEditorHost.ToSharedPtr(), "SDockingTabStack");
+		for (auto& ThisTabStack : TabStacks)
+		{
 			TArray<TSharedRef<SWidget>> DockTabs;
-			FindChildWidgetsOfType(DockTabs, TabStack.ToSharedPtr(), "SDockTab");
+			FindChildWidgetsOfType(DockTabs, ThisTabStack.ToSharedPtr(), "SDockTab");
+
 			for (auto& ThisTab : DockTabs)
 			{
 				TSharedRef<SDockTab> ThisTabCast = StaticCastSharedRef<SDockTab>(ThisTab);
@@ -794,7 +831,7 @@ static bool FindEditorName(const FWidgetPath& WidgetPath, FAIAssistantSlateQuery
 							if (ThisEditor && ThisEditor->GetAssociatedTabManager() == PickedWidgetTabManager)
 							{
 								SlateQueryContext.StructuredContext.EditorName = ThisEditor->GetEditorName().ToString();
-								SlateQueryContext.LastPickedWidget = TabStack.ToSharedPtr();
+								SlateQueryContext.LastPickedWidget = ThisTabStack.ToSharedPtr();
 								FFormatNamedArguments Args;
 								Args.Add(TEXT("AssetEditorName"), FText::FromName(ThisEditor->GetEditorName()));
 								SlateQueryContext.FormattedWindowName = FText::Format(LOCTEXT("EditorName_AssetEditorFormat", " the {AssetEditorName} editor"), Args);
@@ -943,6 +980,11 @@ void UE::AIAssistant::SlateQuerier::QueryAIAssistantAboutSlateWidgetUnderCursor(
 		WidgetPath = FSlateApplication::Get().LocateWindowUnderMouse(CursorPos, Windows, true);
 	}
 
+	if (!WidgetPath.IsValid())
+	{
+		UE_LOG(LogAIAssistant, Warning, TEXT("No valid widget path found to generate widget query."));
+		return;
+	}
 
 	// We build this up, below.
 	
@@ -1092,6 +1134,10 @@ void UE::AIAssistant::SlateQuerier::QueryAIAssistantAboutSlateWidgetUnderCursor(
 		}
 	}
 
+  // DEBUG
+  FString StructuredContextOutput = SlateQueryContext.StructuredContext.ToJson();
+  UE_LOG(LogAIAssistant, Display, TEXT("context json: %s"), *StructuredContextOutput);
+
 	// OPTIONAL - We're done using these members. We can clear them to reduce size as they're
 	// not used below.
 	SlateQueryContext.LastPickedWidget.Reset();
