[AIAssistant]

Added setting for DefaultEditor.ini -
Section "AIAssistant", setting "bIsEnabled", defaulted true.
Only when present and when false, plugin is absent from Editor.

i.e. DefaultEditor.ini -
[AIAssistant]
bIsEnabled=false

diff --git a/Engine/Plugins/Experimental/AIAssistant/Source/AIAssistant/Private/AIAssistant.cpp b/Engine/Plugins/Experimental/AIAssistant/Source/AIAssistant/Private/AIAssistant.cpp
--- a/Engine/Plugins/Experimental/AIAssistant/Source/AIAssistant/Private/AIAssistant.cpp
+++ b/Engine/Plugins/Experimental/AIAssistant/Source/AIAssistant/Private/AIAssistant.cpp
@@ -7,6 +7,7 @@
 #include "Widgets/Docking/SDockTab.h"
 #include "HAL/PlatformApplicationMisc.h"
 #include "Framework/Application/SlateApplication.h"
+#include "Misc/ConfigCacheIni.h"
 
 #include "AIAssistantCommands.h"
 #include "AIAssistantInputProcessor.h"
@@ -18,7 +19,7 @@
 
 
 //
-// Statics.
+// Statics
 //
 
 
@@ -32,6 +33,14 @@ static const FName AIAssistantTabName("AIAssistant");
 
 /*virtual*/ void FAIAssistantModule::StartupModule()
 {
+	GConfig->GetBool(TEXT("AIAssistant"), TEXT("bIsEnabled"), bIsAIAssistantEnabled, GEditorIni);
+
+	if (!bIsAIAssistantEnabled)
+	{
+		return;
+	}
+	
+		
 	FAIAssistantStyle::Initialize();
 	FAIAssistantStyle::ReloadTextures();
 	
@@ -72,6 +81,12 @@ static const FName AIAssistantTabName("AIAssistant");
 
 /*virtual*/ void FAIAssistantModule::ShutdownModule()
 {
+	if (!bIsAIAssistantEnabled)
+	{
+		return;
+	}
+
+	
 	if (FSlateApplication::IsInitialized())
 	{
 		FSlateApplication::Get().UnregisterInputPreProcessor(InputProcessor);
diff --git a/Engine/Plugins/Experimental/AIAssistant/Source/AIAssistant/Public/AIAssistant.h b/Engine/Plugins/Experimental/AIAssistant/Source/AIAssistant/Public/AIAssistant.h
--- a/Engine/Plugins/Experimental/AIAssistant/Source/AIAssistant/Public/AIAssistant.h
+++ b/Engine/Plugins/Experimental/AIAssistant/Source/AIAssistant/Public/AIAssistant.h
@@ -70,6 +70,9 @@ private:
 	
 	TSharedRef<SDockTab> OnSpawnPluginTab(const FSpawnTabArgs& SpawnTabArgs);
 
+
+	bool bIsAIAssistantEnabled = true;
+	
 	TSharedPtr<FUICommandList> PluginCommands;
 	TSharedPtr<FAIAssistantInputProcessor> InputProcessor;
 	TSharedPtr<SAIAssistantWebBrowser> AIAssistantWebBrowserWidget;
