Move the Assistant to a new spot in the Window menu.

diff --git a/Engine/Plugins/Experimental/AIAssistant/Source/AIAssistant/Private/AIAssistant.cpp b/Engine/Plugins/Experimental/AIAssistant/Source/AIAssistant/Private/AIAssistant.cpp
--- a/Engine/Plugins/Experimental/AIAssistant/Source/AIAssistant/Private/AIAssistant.cpp
+++ b/Engine/Plugins/Experimental/AIAssistant/Source/AIAssistant/Private/AIAssistant.cpp
@@ -178,23 +178,13 @@ void FAIAssistantModule::RegisterMenus()
 	FToolMenuOwnerScoped OwnerScoped(this);
 
 	{
-		UToolMenu* Menu = UToolMenus::Get()->ExtendMenu("LevelEditor.MainMenu.Window");
+		UToolMenu* Menu = UToolMenus::Get()->ExtendMenu("MainFrame.MainMenu.Window");
 
-		FToolMenuSection& Section = Menu->FindOrAddSection("WindowLayout");
+		FToolMenuSection& Section = Menu->AddSection("Assistance", LOCTEXT("AIAssistantMenuSection_Label", "Assistance"), FToolMenuInsert("GetContent", EToolMenuInsertType::Before));
 
-		Section.AddMenuEntryWithCommandList(FAIAssistantCommands::Get().OpenAIAssistantTab, PluginCommands);
+		Section.AddMenuEntryWithCommandList(FAIAssistantCommands::Get().OpenAIAssistantTab, PluginCommands, TAttribute<FText>(), TAttribute<FText>(), FSlateIcon("AIAssistantStyle", "AIAssistant.OpenPluginWindow"));
 	}
 
-	{
-		UToolMenu* ToolbarMenu = UToolMenus::Get()->ExtendMenu("LevelEditor.LevelEditorToolBar");
-
-		FToolMenuSection& Section = ToolbarMenu->FindOrAddSection("Settings");
-
-		FToolMenuEntry& Entry = Section.AddEntry(FToolMenuEntry::InitToolBarButton(
-			FAIAssistantCommands::Get().OpenAIAssistantTab));
-
-		Entry.SetCommandList(PluginCommands);
-	}
 }
 
 
