[AI Assistant] Guard against F1 UI queries while in PIE with mouse captured

diff --git a/Engine/Plugins/Experimental/AIAssistant/Source/AIAssistant/Private/AIAssistant.cpp b/Engine/Plugins/Experimental/AIAssistant/Source/AIAssistant/Private/AIAssistant.cpp
--- a/Engine/Plugins/Experimental/AIAssistant/Source/AIAssistant/Private/AIAssistant.cpp
+++ b/Engine/Plugins/Experimental/AIAssistant/Source/AIAssistant/Private/AIAssistant.cpp
@@ -3,11 +3,14 @@
 
 #include "AIAssistant.h"
 
+#include "Editor/EditorEngine.h"
+#include "Engine/GameViewportClient.h"
 #include "ToolMenus.h"
 #include "Widgets/Docking/SDockTab.h"
 #include "HAL/PlatformApplicationMisc.h"
 #include "Framework/Application/SlateApplication.h"
 #include "Misc/ConfigCacheIni.h"
+#include "Slate/SceneViewport.h"
 
 #include "AIAssistantCommands.h"
 #include "AIAssistantInputProcessor.h"
@@ -114,6 +117,24 @@ void FAIAssistantModule::OnOpenPluginTab()
 
 void FAIAssistantModule::OnSlateQuery()
 {
+	UEditorEngine* EditorEngine = Cast<UEditorEngine>(GEngine);
+	if (EditorEngine && EditorEngine->IsPlayingSessionInEditor())
+	{
+		// PIE is currently active
+		if (GEngine->GameViewport)
+		{
+			FSceneViewport* SceneViewport = GEngine->GameViewport->GetGameViewport();
+			if (SceneViewport)
+			{
+				if (SceneViewport->HasMouseCapture())
+				{
+					// If mouse is currently captured by PIE, do nothing.
+					return;
+				}
+			}
+		}
+	}
+
 	OnOpenPluginTab(); // ..idempotent, call before query to ensure tab is open and guarantee a valid AIAssistantWebBrowserWidget
 
 	UE::AIAssistant::SlateQuerier::QueryAIAssistantAboutSlateWidgetUnderCursor();
