[AIAssistant]

Web Connection widget -

Added ability to disable the 'reconnect' ticker, and manual update reconnect method - mainly for tests.

Web Connection test -

New test for the Web Connection widget.

diff --git a/Engine/Plugins/Experimental/AIAssistant/Source/AIAssistant/Private/AIAssistantWebConnectionWidget.cpp b/Engine/Plugins/Experimental/AIAssistant/Source/AIAssistant/Private/AIAssistantWebConnectionWidget.cpp
--- a/Engine/Plugins/Experimental/AIAssistant/Source/AIAssistant/Private/AIAssistantWebConnectionWidget.cpp
+++ b/Engine/Plugins/Experimental/AIAssistant/Source/AIAssistant/Private/AIAssistantWebConnectionWidget.cpp
@@ -28,6 +28,7 @@ static constexpr float ReconnectingTickerCheckPeriodSeconds = 5.0;
 
 void SAIAssistantWebConnectionWidget::Construct(const FArguments& InArgs)
 {
+	bUseReconnectingTicker = InArgs._bUseReconnectingTicker;
 	OnConnectedDelegate = InArgs._OnConnected;
 	OnDisconnectedDelegate = InArgs._OnDisconnected;
 	OnReconnectDelegate = InArgs._OnReconnect;
@@ -156,11 +157,14 @@ void SAIAssistantWebConnectionWidget::StartReconnecting()
 
 
 	LastReconnectAttemptSeconds = FPlatformTime::Seconds();
-	
-	ReconnectingTicker = FTSTicker::GetCoreTicker().AddTicker(
-		FTickerDelegate::CreateSP(this, &SAIAssistantWebConnectionWidget::ReceiveReconnectingTickerUpdate),
-		ReconnectingTickerCheckPeriodSeconds);
 
+	if (bUseReconnectingTicker)
+	{
+		ReconnectingTicker = FTSTicker::GetCoreTicker().AddTicker(
+			FTickerDelegate::CreateSP(this, &SAIAssistantWebConnectionWidget::ReceiveReconnectingTickerUpdate),
+			ReconnectingTickerCheckPeriodSeconds);
+	}
+	
 	ConnectionState = EConnectionState::Reconnecting;
 	OnReconnectDelegate.ExecuteIfBound();
 }
@@ -184,11 +188,11 @@ void SAIAssistantWebConnectionWidget::StopReconnecting()
 }
 
 
-bool SAIAssistantWebConnectionWidget::ReceiveReconnectingTickerUpdate(float DeltaTime)
+void SAIAssistantWebConnectionWidget::UpdateReconnecting()
 {
 	if (ConnectionState == EConnectionState::Connected || ConnectionState == EConnectionState::Disconnected)
 	{
-		return false;	
+		return;	
 	}
 	
 	if (const EConnectionState RequestedConnectionState = (OnRequestConnectionStateFunction ? OnRequestConnectionStateFunction() : EConnectionState::Connected);
@@ -198,8 +202,6 @@ bool SAIAssistantWebConnectionWidget::ReceiveReconnectingTickerUpdate(float Delt
 		
 		ConnectionState = EConnectionState::Connected;
 		OnConnectedDelegate.ExecuteIfBound();
-
-		return false;
 	}
 	else if (RequestedConnectionState == EConnectionState::Disconnected)
 	{
@@ -207,12 +209,15 @@ bool SAIAssistantWebConnectionWidget::ReceiveReconnectingTickerUpdate(float Delt
 	
 		ConnectionState = EConnectionState::Disconnected;
 		OnDisconnectedDelegate.ExecuteIfBound();
-
-		return false;
 	}
+}
 
 
-	return true;
+bool SAIAssistantWebConnectionWidget::ReceiveReconnectingTickerUpdate(float DeltaTime)
+{
+	UpdateReconnecting();
+
+	return (ConnectionState == EConnectionState::Reconnecting);
 }
 
 
diff --git a/Engine/Plugins/Experimental/AIAssistant/Source/AIAssistant/Private/AIAssistantWebConnectionWidget.h b/Engine/Plugins/Experimental/AIAssistant/Source/AIAssistant/Private/AIAssistantWebConnectionWidget.h
--- a/Engine/Plugins/Experimental/AIAssistant/Source/AIAssistant/Private/AIAssistantWebConnectionWidget.h
+++ b/Engine/Plugins/Experimental/AIAssistant/Source/AIAssistant/Private/AIAssistantWebConnectionWidget.h
@@ -38,10 +38,12 @@ public:
 	
 	
 	SLATE_BEGIN_ARGS(SAIAssistantWebConnectionWidget) :
-		_WhenDisconnectedMessage(LOCTEXT("AIAssistantWebConnection_NoInternetConnection", "No Internet Connection"))
+		_WhenDisconnectedMessage(LOCTEXT("AIAssistantWebConnection_NoInternetConnection", "No Internet Connection")),
+		_bUseReconnectingTicker(true)
 		{
 		}
 		SLATE_ARGUMENT(FText, WhenDisconnectedMessage)
+		SLATE_ARGUMENT(bool, bUseReconnectingTicker)
 		SLATE_ARGUMENT(FOnRequestConnectionState, OnRequestConnectionState)
 		SLATE_EVENT(FOnConnected, OnConnected)
 		SLATE_EVENT(FOnDisconnected, OnDisconnected)
@@ -71,6 +73,11 @@ public:
 	 * Stop attempting to acquire connected state.
 	 */
 	void StopReconnecting();
+
+	/**
+	 * Update to check reconnection. Should probably only call this if not using the reconnecting ticker, which calls this automatically.
+	 */
+	void UpdateReconnecting();
 	
 	
 private:
@@ -83,6 +90,7 @@ private:
 	FOnConnected OnConnectedDelegate;
 	FOnDisconnected OnDisconnectedDelegate;
 	FOnReconnect OnReconnectDelegate;
+	bool bUseReconnectingTicker = true;
 	FTSTicker::FDelegateHandle ReconnectingTicker;
 	FOnRequestConnectionState OnRequestConnectionStateFunction;
 	double LastReconnectAttemptSeconds = 0.0;
diff --git a/Engine/Plugins/Experimental/AIAssistant/Source/AIAssistant/Private/Tests/AIAssistantWebConnectionTest.cpp b/Engine/Plugins/Experimental/AIAssistant/Source/AIAssistant/Private/Tests/AIAssistantWebConnectionTest.cpp
new file mode 100644
--- /dev/null
+++ b/Engine/Plugins/Experimental/AIAssistant/Source/AIAssistant/Private/Tests/AIAssistantWebConnectionTest.cpp
@@ -0,0 +1,358 @@
+// Copyright Epic Games, Inc. All Rights Reserved.
+	
+
+#include "AIAssistantTestFlags.h"
+
+#include "AIAssistantWebConnectionWidget.h"
+
+
+#if WITH_DEV_AUTOMATION_TESTS
+
+
+//
+// FAIAssistantFakeWebConnectionWidgetTester
+//
+// Contains a SAIAssistantWebConnectionWidget.
+// Also contains tracking of event calls made to provided lambdas, and other helpers for testing.
+//
+
+
+class FAIAssistantFakeWebConnectionWidgetTester
+{
+public:
+
+	
+	enum class EFakeWebPageLoadState : uint8
+	{
+		LoadPending = 0,
+		LoadError,
+		LoadComplete
+	};
+
+
+	/**
+	 * @param InAutomationTestBase The test we're using this in. 
+	 * @param InFakeWebPageLoadState The simulated web page load state to assume during this test.
+	 */
+	FAIAssistantFakeWebConnectionWidgetTester(FAutomationTestBase& InAutomationTestBase, const EFakeWebPageLoadState InFakeWebPageLoadState) :
+		AutomationTestBase(InAutomationTestBase),
+		FakeWebPageLoadState(InFakeWebPageLoadState)
+	{
+		SAssignNew(AIAssistantWebConnectionWidget, SAIAssistantWebConnectionWidget)
+			.bUseReconnectingTicker(false) // ..we can't use the ticker during these tests
+			.OnRequestConnectionState([this]() -> SAIAssistantWebConnectionWidget::EConnectionState
+				{
+					// This lambda simulates how SAIAssistantWebBrowser would react to being told to reconnect.
+					// We have a fake state that simulates whether connection is possible.
+			
+					if (FakeWebPageLoadState == EFakeWebPageLoadState::LoadError)
+					{
+						return SAIAssistantWebConnectionWidget::EConnectionState::Disconnected;
+					}
+					else if (FakeWebPageLoadState == EFakeWebPageLoadState::LoadComplete)
+					{
+						return SAIAssistantWebConnectionWidget::EConnectionState::Connected;
+					}
+					else
+					{
+						return SAIAssistantWebConnectionWidget::EConnectionState::Reconnecting;
+					}
+				})
+			.OnConnected_Lambda([this]() -> void
+				{
+					++ConnectedEventCount;
+				})
+			.OnDisconnected_Lambda([this]() -> void
+				{
+					++DisconnectedEventCount;
+				})
+			.OnReconnect_Lambda([this]() -> void
+				{
+					++ReconnectEventCount;
+				});
+	}
+
+
+	/**
+	 * @return Internal SAIAssistantWebConnectionWidget used for testing.  
+	 */
+	TSharedPtr<SAIAssistantWebConnectionWidget> GetAIAssistantWebConnectionWidget() const
+	{
+		return AIAssistantWebConnectionWidget;
+	}
+
+
+	/**
+	 * @return Returns current testing step counter, used to control different phases to a test.
+	 */
+	unsigned int GetTestStep() const
+	{
+		return TestStep;
+	}
+
+	/**
+	 * Updates reconnecting state (since the SAIAssistantWebConnectionWidget is not using a ticker), and advances the testing step counter.
+	 */
+	void UpdateReconnectingAndAdvanceTestStep()
+	{
+		AIAssistantWebConnectionWidget->UpdateReconnecting();
+		
+		TestStep++;	
+	}
+
+
+	/**
+	 * Checks expected number of event (lambda call) counts against the number of actual event (lambda call) counts we received during testing.
+	 * @param ExpectedConnectedEventCount Number of expected connected events.
+	 * @param ExpectedDisconnectedEventCount Number of expected disconnected events.
+	 * @param ExpectedReconnectEventCount  Number of expected reconnect events.
+	 * @return Whether all counts.
+	 */
+	bool TestConnectionEventCounts(
+		const unsigned int ExpectedConnectedEventCount,
+		const unsigned int ExpectedDisconnectedEventCount,
+		const unsigned int ExpectedReconnectEventCount) const
+	{
+		auto GetTestWhat = [&](const FString& ValueName) -> FString
+			{
+				return FString::Printf(TEXT("Test_%d_%s"), TestStep, *ValueName);		
+			};
+
+		bool bOk = true;
+		bOk &= AutomationTestBase.TestEqual(GetTestWhat("ConnectedEventCount"), ConnectedEventCount, ExpectedConnectedEventCount);
+		bOk &= AutomationTestBase.TestEqual(GetTestWhat("DisconnectedEventCount"), DisconnectedEventCount, ExpectedDisconnectedEventCount);
+		bOk &= AutomationTestBase.TestEqual(GetTestWhat("ReconnectEventCount"), ReconnectEventCount, ExpectedReconnectEventCount);
+	
+		return bOk;
+	}
+
+	
+private:
+
+	
+	FAutomationTestBase& AutomationTestBase;
+	TSharedPtr<SAIAssistantWebConnectionWidget> AIAssistantWebConnectionWidget;
+	EFakeWebPageLoadState FakeWebPageLoadState = EFakeWebPageLoadState::LoadPending;
+	unsigned int TestStep = 0;
+	unsigned int DisconnectedEventCount = 0, ConnectedEventCount = 0, ReconnectEventCount = 0;
+};
+
+
+//
+// AIAssistantWebConnectionTestDisconnectTwice
+//
+
+
+IMPLEMENT_SIMPLE_AUTOMATION_TEST(
+	AIAssistantWebConnectionTestDisconnectTwice,
+	"AI.Assistant.WebConnection.DisconnectTwice",
+	AIAssistantTest::Flags);
+
+bool AIAssistantWebConnectionTestDisconnectTwice::RunTest(const FString& Parameters)
+{
+	// Disconnect (twice). (Load state doesn't matter.)
+
+	FAIAssistantFakeWebConnectionWidgetTester Tester(*this, 
+		FAIAssistantFakeWebConnectionWidgetTester::EFakeWebPageLoadState::LoadComplete/*but doesn't matter here*/);
+	while (true)
+	{
+		if (Tester.GetTestStep() == 0)
+		{
+			(void) Tester.TestConnectionEventCounts(0, 0, 0);
+
+			Tester.GetAIAssistantWebConnectionWidget()->Disconnect();
+			Tester.UpdateReconnectingAndAdvanceTestStep();
+		}
+		else if (Tester.GetTestStep() == 1)
+		{
+			(void) Tester.TestConnectionEventCounts(0, 1, 0);
+
+			Tester.GetAIAssistantWebConnectionWidget()->Disconnect();
+			Tester.UpdateReconnectingAndAdvanceTestStep();
+		}
+		else if (Tester.GetTestStep() == 2)
+		{
+			(void) Tester.TestConnectionEventCounts(0, 1, 0);
+
+			Tester.UpdateReconnectingAndAdvanceTestStep();
+		}
+		else
+		{
+			break;
+		}
+	}
+
+	
+	return true;
+}
+
+
+//
+// AIAssistantWebConnectionTestReconnectTwiceWhenLoadComplete
+//
+
+
+IMPLEMENT_SIMPLE_AUTOMATION_TEST(
+	AIAssistantWebConnectionTestReconnectTwiceWhenLoadComplete,
+	"AI.Assistant.WebConnection.ReconnectTwiceWhenLoadComplete",
+	AIAssistantTest::Flags);
+
+bool AIAssistantWebConnectionTestReconnectTwiceWhenLoadComplete::RunTest(const FString& Parameters)
+{
+	// Reconnect (twice) in 'load complete' state.
+
+	FAIAssistantFakeWebConnectionWidgetTester Tester(*this,
+		FAIAssistantFakeWebConnectionWidgetTester::EFakeWebPageLoadState::LoadComplete);
+	while (true)
+	{
+		if (Tester.GetTestStep() == 0)
+		{
+			(void) Tester.TestConnectionEventCounts(0, 0, 0);
+		
+			Tester.GetAIAssistantWebConnectionWidget()->Disconnect();
+			Tester.UpdateReconnectingAndAdvanceTestStep();
+		}
+		else if (Tester.GetTestStep() == 1)
+		{
+			(void) Tester.TestConnectionEventCounts(0, 1, 0);
+		
+			Tester.GetAIAssistantWebConnectionWidget()->StartReconnecting();
+			Tester.UpdateReconnectingAndAdvanceTestStep();
+		}
+		else if (Tester.GetTestStep() == 2)
+		{
+			(void) Tester.TestConnectionEventCounts(1, 1, 1);
+		
+			Tester.GetAIAssistantWebConnectionWidget()->StartReconnecting();
+			Tester.UpdateReconnectingAndAdvanceTestStep();
+		}
+		else if (Tester.GetTestStep() == 3)
+		{
+			(void) Tester.TestConnectionEventCounts(1, 1, 1);
+			
+			Tester.UpdateReconnectingAndAdvanceTestStep();
+		}
+		else
+		{
+			break;
+		}
+	}
+
+
+	return true;
+};
+
+
+//
+// AIAssistantWebConnectionTestReconnectTwiceWhenLoadError
+//
+
+
+IMPLEMENT_SIMPLE_AUTOMATION_TEST(
+	AIAssistantWebConnectionTestReconnectTwiceWhenLoadError,
+	"AI.Assistant.WebConnection.ReconnectTwiceWhenLoadError",
+	AIAssistantTest::Flags);
+
+bool AIAssistantWebConnectionTestReconnectTwiceWhenLoadError::RunTest(const FString& Parameters)
+{
+	// Reconnect (twice) in 'load error' state.
+
+	FAIAssistantFakeWebConnectionWidgetTester Tester(*this,
+		FAIAssistantFakeWebConnectionWidgetTester::EFakeWebPageLoadState::LoadError);
+	while (true)
+	{
+		if (Tester.GetTestStep() == 0)
+		{
+			(void) Tester.TestConnectionEventCounts(0, 0, 0);
+		
+			Tester.GetAIAssistantWebConnectionWidget()->Disconnect();
+			Tester.UpdateReconnectingAndAdvanceTestStep();
+		}
+		else if (Tester.GetTestStep() == 1)
+		{
+			(void) Tester.TestConnectionEventCounts(0, 1, 0);
+		
+			Tester.GetAIAssistantWebConnectionWidget()->StartReconnecting();
+			Tester.UpdateReconnectingAndAdvanceTestStep();
+		}
+		else if (Tester.GetTestStep() == 2)
+		{
+			(void) Tester.TestConnectionEventCounts(0, 2, 1);
+		
+			Tester.GetAIAssistantWebConnectionWidget()->StartReconnecting();
+			Tester.UpdateReconnectingAndAdvanceTestStep();
+		}
+		else if (Tester.GetTestStep() == 3)
+		{
+			(void) Tester.TestConnectionEventCounts(0, 3, 2);
+			
+			Tester.UpdateReconnectingAndAdvanceTestStep();
+		}
+		else
+		{
+			break;
+		}
+	}
+
+	
+	return true;
+}
+
+
+//
+// AIAssistantWebConnectionTestReconnectTwiceWhenLoadPending
+//
+
+
+IMPLEMENT_SIMPLE_AUTOMATION_TEST(
+	AIAssistantWebConnectionTestReconnectTwiceWhenLoadPending,
+	"AI.Assistant.WebConnection.ReconnectTwiceWhenLoadPending",
+	AIAssistantTest::Flags);
+
+bool AIAssistantWebConnectionTestReconnectTwiceWhenLoadPending::RunTest(const FString& Parameters)
+{
+	// Reconnect (twice) in 'load pending' state.
+
+	FAIAssistantFakeWebConnectionWidgetTester Tester(*this,
+		FAIAssistantFakeWebConnectionWidgetTester::EFakeWebPageLoadState::LoadPending);
+	while (true)
+	{
+		if (Tester.GetTestStep() == 0)
+		{
+			(void) Tester.TestConnectionEventCounts(0, 0, 0);
+		
+			Tester.GetAIAssistantWebConnectionWidget()->Disconnect();
+			Tester.UpdateReconnectingAndAdvanceTestStep();
+		}
+		else if (Tester.GetTestStep() == 1)
+		{
+			(void) Tester.TestConnectionEventCounts(0, 1, 0);
+		
+			Tester.GetAIAssistantWebConnectionWidget()->StartReconnecting();
+			Tester.UpdateReconnectingAndAdvanceTestStep();
+		}
+		else if (Tester.GetTestStep() == 2)
+		{
+			(void) Tester.TestConnectionEventCounts(0, 1, 1);
+		
+			Tester.GetAIAssistantWebConnectionWidget()->StartReconnecting();
+			Tester.UpdateReconnectingAndAdvanceTestStep();
+		}
+		else if (Tester.GetTestStep() == 3)
+		{
+			(void) Tester.TestConnectionEventCounts(0, 1, 1);
+		
+			Tester.UpdateReconnectingAndAdvanceTestStep();
+		}
+		else
+		{
+			break;
+		}
+	}
+
+	
+	return true;
+}
+
+
+#endif
