[AI Assistant] Do not invoke Tooltip query if mouse cursor is not over any editor window.

#rnx

diff --git a/Engine/Plugins/Experimental/AIAssistant/Source/AIAssistant/Private/AIAssistant.cpp b/Engine/Plugins/Experimental/AIAssistant/Source/AIAssistant/Private/AIAssistant.cpp
--- a/Engine/Plugins/Experimental/AIAssistant/Source/AIAssistant/Private/AIAssistant.cpp
+++ b/Engine/Plugins/Experimental/AIAssistant/Source/AIAssistant/Private/AIAssistant.cpp
@@ -191,10 +191,12 @@ void FAIAssistantModule::OnSlateQuery()
 	}
 
 	FWidgetPath WidgetPath = UE::AIAssistant::SlateQuerier::GetWidgetPathUnderCursor();
-
-	// ..idempotent, call before query to ensure tab is open and guarantee a valid AIAssistantWebBrowserWidget
-	OnOpenPluginTab(WidgetPath);
-	UE::AIAssistant::SlateQuerier::QueryAIAssistantAboutSlateWidget(WidgetPath);
+	if (WidgetPath.IsValid())
+	{
+		// ..idempotent, call before query to ensure tab is open and guarantee a valid AIAssistantWebBrowserWidget
+		OnOpenPluginTab(WidgetPath);
+		UE::AIAssistant::SlateQuerier::QueryAIAssistantAboutSlateWidget(WidgetPath);
+	}
 }
 
 
