[AI Assistant] Improve UI queries on Viewports
- special prompt handling for UI queries on AI Assistant window
- start refactoring into smaller query functions
- add special handling for property filter checkboxes in Details panel
- UI queries in Viewports will ask about selected/edited object if any
- UI queries in LevelEditor Viewport will ask about hovered object if nothing is selected

#rnx

diff --git a/Engine/Plugins/Experimental/AIAssistant/Source/AIAssistant/Private/AIAssistantSlateQuerier.cpp b/Engine/Plugins/Experimental/AIAssistant/Source/AIAssistant/Private/AIAssistantSlateQuerier.cpp
--- a/Engine/Plugins/Experimental/AIAssistant/Source/AIAssistant/Private/AIAssistantSlateQuerier.cpp
+++ b/Engine/Plugins/Experimental/AIAssistant/Source/AIAssistant/Private/AIAssistantSlateQuerier.cpp
@@ -4,14 +4,20 @@
 #include "AIAssistantSlateQuerier.h"
 
 #include "EditorModes.h"
+#include "Containers/Set.h"
+#include "Containers/ArrayView.h"
 #include "Internationalization/Regex.h"
 #include "LevelEditorSubsystem.h"
+#include "LevelEditorViewport.h"
+#include "Settings/LevelEditorViewportSettings.h"
 #include "EditorModeManager.h"
+#include "EngineUtils.h"
 #include "IDetailsView.h"
 #include "Interfaces/IMainFrameModule.h"
 #include "Serialization/JsonSerializable.h"
 #include "Serialization/JsonSerializerMacros.h"
 #include "Subsystems/AssetEditorSubsystem.h"
+#include "Subsystems/EditorActorSubsystem.h"
 #include "Toolkits/BaseToolkit.h"
 #include "Widgets/Text/STextBlock.h"
 #include "Widgets/Text/SRichTextBlock.h"
@@ -82,8 +88,11 @@ struct FAIAssistantSlateQueryContext
 	FText TextUnderCursor = FText();
 	FText ItemName = FText();
 	FText ItemDescriptor = FText();
+	FString SelectedActorName = FString();
+	FString SelectedActorClass = FString();
 	bool bIsUIWidget = false;
 	bool bIsObject = false;
+	bool bIsInAssistantPanel = false;
 
 	FText GeneratedQuery = FText();
 	FText GeneratedQueryInstructions = FText();
@@ -96,39 +105,68 @@ struct FAIAssistantSlateQueryContext
 // Statics
 //
 
-
-static TSharedRef<SWidget> FindClosestWidgetOfType(const FWidgetPath& WidgetPathToTest, const FName& WidgetType)
+// Search backwards from end of WidgetPathToTest and return first widget that matches any type in WidgetTypes.
+// Return null ptr if not found.
+static TSharedPtr<SWidget> FindClosestWidgetPtrOfTypes(const FWidgetPath& WidgetPathToTest, const TSet<FName> WidgetTypes)
 {
+	TSharedPtr<SWidget> ClosestWidget;
 	for (int32 WidgetIndex = WidgetPathToTest.Widgets.Num() - 1; WidgetIndex >= 0; --WidgetIndex)
 	{
 		const FArrangedWidget* ArrangedWidget = &WidgetPathToTest.Widgets[WidgetIndex];
 		const TSharedRef<SWidget>& ThisWidget = ArrangedWidget->Widget;
-		if (ThisWidget->GetType() == WidgetType)
+		if (WidgetTypes.Contains(ThisWidget->GetType()))
 		{
-			return ThisWidget;
+			ClosestWidget = ThisWidget.ToSharedPtr();
+			break;
 		}
 	}
-	
-	return SNullWidget::NullWidget;
+	return ClosestWidget;
 }
 
+// Search backwards from end of WidgetPathToTest and return first widget that matches WidgetType.  Return null
+// ptr if not found.
+static TSharedPtr<SWidget> FindClosestWidgetPtrOfType(const FWidgetPath& WidgetPathToTest, const FName& WidgetType)
+{
+	const TSet<FName> WidgetTypes = TSet<FName>({WidgetType});
+	return FindClosestWidgetPtrOfTypes(WidgetPathToTest, WidgetTypes);
+}
 
-static TSharedRef<SWidget> FindFirstWidgetOfType(const FWidgetPath& WidgetPathToTest, const FName& WidgetType)
+// Search backwards from end of WidgetPathToTest and return first widget that matches a type in WidgetTypes,
+// testing each type in order.  Return null ptr if not found.
+static TSharedPtr<SWidget> FindClosestWidgetPtrOfSortedTypes(
+	const FWidgetPath& WidgetPathToTest, const TConstArrayView<FName>& WidgetTypes)
 {
+	for (FName WidgetType : WidgetTypes)
+	{
+		TSharedPtr<SWidget> ClosestWidget = FindClosestWidgetPtrOfType(WidgetPathToTest, WidgetType);
+		if (ClosestWidget.IsValid())
+		{
+			return ClosestWidget;
+		}
+	}
+
+	return nullptr;
+}
+
+// Search forward from start of WidgetPathToTest and return first widget that matches WidgetType.  Return null
+// ptr if not found.
+static TSharedPtr<SWidget> FindFirstWidgetPtrOfType(const FWidgetPath& WidgetPathToTest, const FName& WidgetType)
+{
+	TSharedPtr<SWidget> FirstWidget;
 	for (int32 WidgetIndex = 0; WidgetIndex < WidgetPathToTest.Widgets.Num(); WidgetIndex++)
 	{
 		const FArrangedWidget* ArrangedWidget = &WidgetPathToTest.Widgets[WidgetIndex];
 		const TSharedRef<SWidget>& ThisWidget = ArrangedWidget->Widget;
 		if (ThisWidget->GetType() == WidgetType)
 		{
-			return ThisWidget;
+			FirstWidget = ThisWidget.ToSharedPtr();
+			break;
 		}
 	}
-	
-	return SNullWidget::NullWidget;
+	return FirstWidget;
 }
 
-
+// Search children of WidgetToTest for any widgets that match WidgetType; return all found widgets in OutWidgets.
 static void FindChildWidgetsOfType(TArray<TSharedRef<SWidget>>& OutWidgets, TSharedPtr<SWidget> WidgetToTest, const FName& WidgetType)
 {
 	if (WidgetToTest.IsValid())
@@ -145,8 +183,8 @@ static void FindChildWidgetsOfType(TArray<TSharedRef<SWidget>>& OutWidgets, TSha
 	}
 }
 
-
-static TSharedRef<SWidget> FindChildWidgetOfType(const TSharedPtr<SWidget> WidgetToTest, const FName& WidgetType)
+// Search children of WidgetToTest and return first widget that matches WidgetType.  Return null ptr if not found.
+static TSharedPtr<SWidget> FindChildWidgetOfType(const TSharedPtr<SWidget> WidgetToTest, const FName& WidgetType)
 {
 	if (WidgetToTest.IsValid())
 	{
@@ -155,23 +193,23 @@ static TSharedRef<SWidget> FindChildWidgetOfType(const TSharedPtr<SWidget> Widge
 			TSharedRef<SWidget> ThisWidget = WidgetToTest->GetChildren()->GetChildAt(ChildIdx);
 			if (ThisWidget->GetType() == WidgetType)
 			{
-				return ThisWidget;
+				return ThisWidget.ToSharedPtr();
 			}
 			else
 			{
-				TSharedRef<SWidget> ChildWidget = FindChildWidgetOfType(ThisWidget.ToSharedPtr(), WidgetType);
-				if (ChildWidget->GetType() == WidgetType)
+				TSharedPtr<SWidget> ChildWidget = FindChildWidgetOfType(ThisWidget.ToSharedPtr(), WidgetType);
+				if (ChildWidget.IsValid() && ChildWidget->GetType() == WidgetType)
 				{
 					return ChildWidget;
 				}
 			}
 		}
 	}
-	
-	return SNullWidget::NullWidget;
+	return TSharedPtr<SWidget>();
 }
 
-
+// Search children of WidgetToTest and return text from first one found with alpha text.  Return empty FText
+// if not found.
 static FText FindChildWidgetWithText(const TSharedPtr<SWidget> WidgetToTest)
 {
 	// get all children and see if any are text widgets.
@@ -206,7 +244,7 @@ static FText FindChildWidgetWithText(const TSharedPtr<SWidget> WidgetToTest)
 	return FText();
 }
 
-
+// Return first active tooltip text found.  Return empty FText if not found.
 static FText FindCurrentToolTipText()
 {
 	for (const TArray<TSharedRef<SWindow>> AllWindows = FSlateApplication::Get().GetTopLevelWindows();
@@ -221,46 +259,47 @@ static FText FindCurrentToolTipText()
 	return FText();
 }
 
-
-static TSharedRef<SWidget> FindClosestMenuItem(const FWidgetPath& WidgetPath)
+// Search backwards from end of WidgetPath and return first SMenuEntryBlock or SWidgetBlock widget found.
+// Return null ptr if not found.
+static TSharedPtr<SWidget> FindClosestMenuItem(const FWidgetPath& WidgetPath)
 {
+	TSharedPtr<SWidget> MenuItemBlock;
 	if (WidgetPath.IsValid())
 	{
 		// is it a menu item? Look up for an SMenuEntryBlock or SWidgetBlock.
-		TSharedRef<SWidget> MenuItemBlock = FindClosestWidgetOfType(WidgetPath, "SMenuEntryBlock");
-		if (MenuItemBlock->GetType() == "SNullWidgetContent")
-		{
-			MenuItemBlock = FindClosestWidgetOfType(WidgetPath, "SWidgetBlock");
-		}
+		static const TStaticArray<FName, 2> MenuTypes = {"SMenuEntryBlock", "SWidgetBlock"};
+		MenuItemBlock = FindClosestWidgetPtrOfSortedTypes(WidgetPath, MenuTypes);
+
 		// disabled menu items *contain* the menu entry block
-		if (MenuItemBlock->GetType() == "SNullWidgetContent")
+		if (!MenuItemBlock.IsValid())
 		{
 			if (WidgetPath.GetLastWidget()->GetChildren()->Num() > 0)
 			{
 				for (int32 ChildIdx = 0; ChildIdx < WidgetPath.GetLastWidget()->GetChildren()->Num(); ChildIdx++)
 				{
 					TSharedRef<SWidget> ThisWidget = WidgetPath.GetLastWidget()->GetChildren()->GetChildAt(ChildIdx);
-					if (ThisWidget->GetType() == "SMenuEntryBlock" || ThisWidget->GetType() == "SWidgetBlock")
+					for (FName MenuType : MenuTypes)
 					{
-						MenuItemBlock = ThisWidget;
+						if (ThisWidget->GetType() == MenuType)
+						{
+							MenuItemBlock = ThisWidget.ToSharedPtr();
+							break;
+						}
 					}
 				}
 			}
 		}
-		if (MenuItemBlock->GetType() != "SNullWidgetContent")
-		{
-			return MenuItemBlock;
-		}
 	}
 
-	return SNullWidget::NullWidget;
+	return MenuItemBlock;
 }
 
-
+// For menu items, we want to generate the WidgetPath context based on the root widget that spawned the menu,
+// not on the actual menu item itself.
 static FWidgetPath FindContextWidgetPath(const FWidgetPath& WidgetPath)
 {
-	if (const TSharedRef<SWidget> MenuItem = FindClosestMenuItem(WidgetPath);
-		MenuItem->GetType() != "SNullWidgetContent")
+	const TSharedPtr<SWidget> MenuItem = FindClosestMenuItem(WidgetPath);
+	if (MenuItem.IsValid())
 	{
 		TSharedPtr<SWidget> RootWidget = FSlateApplication::Get().GetMenuHostWidget();
 		if (RootWidget.IsValid())
@@ -273,58 +312,45 @@ static FWidgetPath FindContextWidgetPath(const FWidgetPath& WidgetPath)
 	return WidgetPath;
 }
 
-
-static void GenerateModeToolContext(const FWidgetPath& WidgetPath, FAIAssistantSlateQueryContext& SlateQueryContext)
+// Return true and fill in SlateQueryContext if queried item is in viewport.
+static bool ItemIsInViewport(const FWidgetPath& WidgetPath, FAIAssistantSlateQueryContext& SlateQueryContext)
 {
-	// if we're in a level editor mode, include which tool is active
-	TSharedRef<SWidget> LevelEditor = FindClosestWidgetOfType(WidgetPath, "SLevelEditor");
-	if (LevelEditor->GetType() != "SNullWidgetContent")
+	TSharedPtr<SWidget> Viewport = FindClosestWidgetPtrOfType(WidgetPath, "SViewport");
+	if (Viewport.IsValid())
 	{
-		FText ToolString;
-		ULevelEditorSubsystem* LevelEditorSubsystem = GEditor->GetEditorSubsystem<ULevelEditorSubsystem>();
-		if (FEditorModeTools* ModeManager = LevelEditorSubsystem->GetLevelEditorModeManager())
+		// Web browser widgets also have Viewports!
+		// Ignore web browser views except for the special case of AI Assistant window or panel.
+		TSharedPtr<SWidget> WebBrowser = Viewport->GetParentWidget();
+		if (WebBrowser.IsValid() && WebBrowser->GetType() == "SWebBrowserView")
 		{
-			UAssetEditorSubsystem* AssetEditorSubsystem = GEditor->GetEditorSubsystem<UAssetEditorSubsystem>();
-			TArray<FEditorModeInfo> ModeInfos = AssetEditorSubsystem->GetEditorModeInfoOrderedByPriority();
-			for (const FEditorModeInfo& ModeInfo : ModeInfos)
+			TSharedPtr<SWidget> AssistantPane = FindClosestWidgetPtrOfType(WidgetPath, "SAIAssistantWebBrowser");
+			if (AssistantPane.IsValid())
 			{
-				FString ThisModeID = ModeInfo.ID.ToString();
-				if (UEdMode* EdMode = ModeManager->GetActiveScriptableMode(ModeInfo.ID))
-				{
-					TWeakPtr<FModeToolkit> ModeToolkit = EdMode->GetToolkit();
-					if (ModeToolkit.IsValid())
-					{
-						FText ToolDisplayName = ModeToolkit.Pin()->GetActiveToolDisplayName();
-						if (!ToolDisplayName.IsEmpty())
-						{
-							SlateQueryContext.StructuredContext.ActiveToolName = ToolDisplayName.ToString();
-							SlateQueryContext.StructuredContext.EditorMode = ModeInfo.Name.ToString();
-						}
-					}
-				}
+				SlateQueryContext.bIsInAssistantPanel = true;
+				SlateQueryContext.LastPickedWidget = AssistantPane;
+				return true;
+			}
+		}
+		else
+		{
+			// Selected actors are filled in by the IsInEditor function.
+			if (!SlateQueryContext.SelectedActorName.IsEmpty())
+			{
+				SlateQueryContext.bIsObject = true;
+				SlateQueryContext.LastPickedWidget = Viewport;
+				SlateQueryContext.ItemName = FText::FromString(SlateQueryContext.SelectedActorName);
+				SlateQueryContext.ItemDescriptor = LOCTEXT("ItemDescriptor_SelectedActor", "actor");
+				SlateQueryContext.StructuredContext.ActorClass = SlateQueryContext.SelectedActorClass;
+				return true;
 			}
 		}
 	}
+	return false;
 }
 
-
-static bool FindItemName(const FWidgetPath& WidgetPath, FAIAssistantSlateQueryContext& SlateQueryContext)
+// Return true and fill in SlateQueryContext if queried item is a graph node.
+static bool ItemIsGraphNode(const FWidgetPath& WidgetPath, FAIAssistantSlateQueryContext& SlateQueryContext)
 {
-	SlateQueryContext.ItemDescriptor = LOCTEXT("ItemDescriptor_Generic", "control");
-	TSet<FName> ButtonTypes = TSet<FName>({
-		"SButton",
-		"SPrimaryButton",
-		"SCheckbox"
-		});
-
-	FText ChildText = FText();
-
-	if (!WidgetPath.IsValid())
-	{
-		return false;
-	}
-
-	// Is it a graph node?
 	// Look up for an SGraphEditor and then down for an SGraphPanel. Whatever is the child of that SGraphPanel is our SGraphNode.
 	// Note that this rests on the assumption that SGraphPanel only contains children that are SGraphNodes. This assumption is also
 	// made in the code of SGraphPanel itself.
@@ -343,176 +369,174 @@ static bool FindItemName(const FWidgetPath& WidgetPath, FAIAssistantSlateQueryCo
 				{
 					const FArrangedWidget* NodeArrangedWidget = &WidgetPath.Widgets[EditorDescendantIndex+1];
 					const TSharedRef<SWidget>& ThisNodeWidget = NodeArrangedWidget->Widget;
-					if (ThisNodeWidget->GetType() == "SNiagaraOverviewStackNode")
-					{
-						// Niagara emitters don't typically have useful node names
-						continue;
-					}
 					const TSharedPtr<SGraphNode> AsGraphNode = StaticCastSharedPtr<SGraphNode>(ThisNodeWidget.ToSharedPtr());
 					SlateQueryContext.LastPickedWidget = ThisNodeWidget.ToSharedPtr();
 					SlateQueryContext.ItemName = AsGraphNode->GetNodeObj()->GetNodeTitle(ENodeTitleType::MenuTitle);
 					SlateQueryContext.ItemDescriptor = LOCTEXT("ItemDescriptor_GraphNode", "graph node");
+					return true;
 				}
 			}
 		}
 	}
+	return false;
+}
 
-
-	// Is it a button? Look up for an SButton and then down for the text.
-	TSharedPtr<SWidget> Button;
-	for (int32 WidgetIndex = WidgetPath.Widgets.Num() - 1; WidgetIndex >= 0; --WidgetIndex)
+// Fill query context structure for buttons.
+static bool FillButtonQueryContext(const TSharedPtr<SWidget> Button,
+	FAIAssistantSlateQueryContext& SlateQueryContext, FText ItemDescriptor)
+{
+	FText ChildText = FindChildWidgetWithText(Button);
+	if (!ChildText.IsEmpty() || !SlateQueryContext.CurrentToolTipText.IsEmpty())
 	{
-		const FArrangedWidget* ArrangedWidget = &WidgetPath.Widgets[WidgetIndex];
-		const TSharedRef<SWidget>& ThisWidget = ArrangedWidget->Widget;
-		if (ButtonTypes.Contains(ThisWidget->GetType()))
+		if (!ChildText.IsEmpty())
 		{
-			Button = ThisWidget.ToSharedPtr();
+			SlateQueryContext.ItemName = ChildText;
 		}
+		else
+		{
+			// The tool tip can be used in this case.
+			SlateQueryContext.ItemName = SlateQueryContext.CurrentToolTipText;
+		}
+		SlateQueryContext.ItemDescriptor = ItemDescriptor;
+		SlateQueryContext.LastPickedWidget = Button;
+		SlateQueryContext.bIsUIWidget = true;
+		return true;
 	}
+	return false;
+}
+
+// Return true and fill in SlateQueryContext if queried item is a button type as specified in
+// ButtonWidgetTypes.
+static bool ItemIsButtonType(const FWidgetPath& WidgetPath, FAIAssistantSlateQueryContext& SlateQueryContext,
+	const TSet<FName> ButtonWidgetTypes, FText ItemDescriptor)
+{
+	TSharedPtr<SWidget> Button = FindClosestWidgetPtrOfTypes(WidgetPath, ButtonWidgetTypes);
+
 	// Disabled buttons *contain* the menu entry block.
 	if (!Button.IsValid() && WidgetPath.GetLastWidget()->GetChildren()->Num() > 0)
 	{
 		for (int32 ChildIdx = 0; ChildIdx < WidgetPath.GetLastWidget()->GetChildren()->Num(); ChildIdx++)
 		{
 			TSharedRef<SWidget> ThisWidget = WidgetPath.GetLastWidget()->GetChildren()->GetChildAt(ChildIdx);
-			if (ButtonTypes.Contains(ThisWidget->GetType()))
+			if (ButtonWidgetTypes.Contains(ThisWidget->GetType()))
 			{
 				Button = ThisWidget.ToSharedPtr();
 			}
 		}
 	}
-	if (Button.IsValid())
-	{
-		ChildText = FindChildWidgetWithText(Button);
-		if (!ChildText.IsEmpty() || !SlateQueryContext.CurrentToolTipText.IsEmpty())
-		{
-			if (!ChildText.IsEmpty())
-			{
-				SlateQueryContext.ItemName = ChildText;
-			}
-			else
-			{
-				// The tool tip can be used in this case.
-				SlateQueryContext.ItemName = SlateQueryContext.CurrentToolTipText;
-			}
-			SlateQueryContext.ItemDescriptor = LOCTEXT("ItemDescriptor_Button", "button");
-			SlateQueryContext.LastPickedWidget = Button;
-			SlateQueryContext.bIsUIWidget = true;
-		}
-	}
 
+	return Button.IsValid() && FillButtonQueryContext(Button, SlateQueryContext, ItemDescriptor);
+}
 
-	// Is it a toolbar button? Look for an SToolbarButtonBlock.
-	TSharedRef<SWidget> ToolbarItemBlock = FindClosestWidgetOfType(WidgetPath, "SToolBarButtonBlock");
-	if (ToolbarItemBlock->GetType() == "SNullWidgetContent")
-	{
-		ToolbarItemBlock = FindClosestWidgetOfType(WidgetPath, "SToolBarComboButtonBlock");
-	}
-	// Disabled items *contain* the entry block.
-	if (ToolbarItemBlock->GetType() == "SNullWidgetContent")
-	{
-		if (WidgetPath.GetLastWidget()->GetChildren()->Num() > 0)
-		{
-			for (int32 ChildIdx = 0; ChildIdx < WidgetPath.GetLastWidget()->GetChildren()->Num(); ChildIdx++)
-			{
-				TSharedRef<SWidget> ThisWidget = WidgetPath.GetLastWidget()->GetChildren()->GetChildAt(ChildIdx);
-				if (ThisWidget->GetType() == "SToolBarButtonBlock" || ThisWidget->GetType() == "SToolBarComboButtonBlock")
-				{
-					ToolbarItemBlock = ThisWidget;
-				}
-			}
-		}
-	}
-	if (ToolbarItemBlock->GetType() != "SNullWidgetContent")
-	{
-		ChildText = FindChildWidgetWithText(ToolbarItemBlock);
-		if (!ChildText.IsEmpty() || !SlateQueryContext.CurrentToolTipText.IsEmpty())
-		{
-			if (!ChildText.IsEmpty())
-			{
-				SlateQueryContext.ItemName = ChildText;
-			}
-			else
-			{
-				// The tool tip can be used in this case.
-				SlateQueryContext.ItemName = SlateQueryContext.CurrentToolTipText;
-			}
-			SlateQueryContext.ItemDescriptor = LOCTEXT("ItemDescriptor_ToolbarButton", "toolbar button");
-			SlateQueryContext.LastPickedWidget = ToolbarItemBlock.ToSharedPtr();
-			SlateQueryContext.bIsUIWidget = true;
-		}
-	}
+// Return true and fill in SlateQueryContext if queried item is a button.
+static bool ItemIsButton(const FWidgetPath& WidgetPath, FAIAssistantSlateQueryContext& SlateQueryContext)
+{
+	static const TSet<FName> ToolbarButtonTypes = TSet<FName>({
+		"SToolBarButtonBlock", "SToolBarComboButtonBlock"
+		});
 
+	static const TSet<FName> SimpleButtonTypes = TSet<FName>({
+			"SButton",
+			"SPrimaryButton",
+			"SCheckBox"
+		});
 
-	// Is it a search/filter field?
-	FText HintText;
-	TSharedRef<SWidget> SearchBox = FindClosestWidgetOfType(WidgetPath, "SSearchBox");
-	if (SearchBox->GetType() != "SNullWidgetContent")
-	{
-		TSharedRef<SSearchBox> ThisSearchBoxCast = StaticCastSharedRef<SSearchBox>(SearchBox);
-		HintText = ThisSearchBoxCast->GetHintText();
-	}
-	else
+	return (ItemIsButtonType(WidgetPath, SlateQueryContext, ToolbarButtonTypes,
+				LOCTEXT("ItemDescriptor_ToolbarButton", "toolbar button")) ||
+			ItemIsButtonType(WidgetPath, SlateQueryContext, SimpleButtonTypes,
+				LOCTEXT("ItemDescriptor_Button", "button")));
+}
+
+// Return true and fill in SlateQueryContext if queried item is a search box.
+static bool ItemIsSearchBox(const FWidgetPath& WidgetPath, FAIAssistantSlateQueryContext& SlateQueryContext)
+{
+	static const TSet<FName> SearchBoxTypes = TSet<FName>({"SSearchBox", "SFilterSearchBox"});
+	TSharedPtr<SWidget> SearchBox = FindClosestWidgetPtrOfTypes(WidgetPath, SearchBoxTypes);
+
+	if (SearchBox.IsValid())
 	{
-		SearchBox = FindClosestWidgetOfType(WidgetPath, "SFilterSearchBox");
-		if (SearchBox->GetType() != "SNullWidgetContent")
+		FText HintText;
+		if (SearchBox->GetType() == "SSearchBox")
 		{
-			TSharedRef<SWidget> EditableText = FindChildWidgetOfType(SearchBox, "SEditableText");
-			if (EditableText->GetType() != "SNullWidgetContent")
+			TSharedPtr<SSearchBox> ThisSearchBoxCast = StaticCastSharedPtr<SSearchBox>(SearchBox);
+			HintText = ThisSearchBoxCast->GetHintText();
+		}
+		else
+		{
+			TSharedPtr<SWidget> EditableText = FindChildWidgetOfType(SearchBox, "SEditableText");
+			if (EditableText)
 			{
-				TSharedRef<SEditableText> ThisEditableTextCast = StaticCastSharedRef<SEditableText>(EditableText);
+				TSharedPtr<SEditableText> ThisEditableTextCast = StaticCastSharedPtr<SEditableText>(EditableText);
 				HintText = ThisEditableTextCast->GetHintText();
 			}
 		}
-	}
-	if (SearchBox->GetType() != "SNullWidgetContent")
-	{
 		SlateQueryContext.ItemName = FText::Format(LOCTEXT("SearchBoxFormat", "{0}"),
 			(HintText.IsEmpty()) ? LOCTEXT("SearchBoxDefaultName", "default") : HintText);
 		SlateQueryContext.ItemDescriptor = LOCTEXT("ItemDescriptor_SearchBox", "search box");
-		SlateQueryContext.LastPickedWidget = SearchBox.ToSharedPtr();
+		SlateQueryContext.LastPickedWidget = SearchBox;
 		SlateQueryContext.bIsUIWidget = true;
+		return true;
 	}
 
-	// Is it a console input box?
-	TSharedRef<SWidget> ConsoleInputBox = FindClosestWidgetOfType(WidgetPath, "SConsoleInputBox");
-	if (ConsoleInputBox->GetType() != "SNullWidgetContent")
+	return false;
+}
+
+// Return true and fill in SlateQueryContext if queried item is a console input box.
+static bool ItemIsConsoleInputBox(const FWidgetPath& WidgetPath, FAIAssistantSlateQueryContext& SlateQueryContext)
+{
+	TSharedPtr<SWidget> ConsoleInputBox = FindClosestWidgetPtrOfType(WidgetPath, "SConsoleInputBox");
+	if (ConsoleInputBox.IsValid())
 	{
-		TSharedRef<SWidget> EditableText = FindChildWidgetOfType(ConsoleInputBox, "SMultiLineEditableTextBox");
-		if (EditableText->GetType() != "SNullWidgetContent")
+		FText HintText;
+		TSharedPtr<SWidget> EditableText = FindChildWidgetOfType(ConsoleInputBox, "SMultiLineEditableTextBox");
+		if (EditableText.IsValid())
 		{
-			TSharedRef<SMultiLineEditableTextBox> ThisMultiLineEditableTextBoxCast = StaticCastSharedRef<SMultiLineEditableTextBox>(EditableText);
+			TSharedPtr<SMultiLineEditableTextBox> ThisMultiLineEditableTextBoxCast = StaticCastSharedPtr<SMultiLineEditableTextBox>(EditableText);
 			HintText = ThisMultiLineEditableTextBoxCast->GetHintText();
 		}
 		SlateQueryContext.ItemName = FText::Format(LOCTEXT("ConsoleInputBoxFormat", "{0}"),
 			(HintText.IsEmpty() ? LOCTEXT("ConsoleInputBoxDefaultName","Console") : HintText));
 		SlateQueryContext.ItemDescriptor = LOCTEXT("ItemDescriptor_ConsoleInputBox", "input box");
-		SlateQueryContext.LastPickedWidget = ConsoleInputBox.ToSharedPtr();
+		SlateQueryContext.LastPickedWidget = ConsoleInputBox;
 		SlateQueryContext.bIsUIWidget = true;
+		return true;
 	}
+	return false;
+}
 
-	// Is it a menu item? Look up for an SMenuEntryBlock or SWidgetBlock and then down for the text.
-	TSharedRef<SWidget> MenuItemBlock = FindClosestMenuItem(WidgetPath);
-	if (MenuItemBlock->GetType() != "SNullWidgetContent")
+// Return true and fill in SlateQueryContext if queried item is a menu item.
+static bool ItemIsMenuItem(const FWidgetPath& WidgetPath, FAIAssistantSlateQueryContext& SlateQueryContext)
+{
+	// Look up for an SMenuEntryBlock or SWidgetBlock and then down for the text.
+	TSharedPtr<SWidget> MenuItemBlock = FindClosestMenuItem(WidgetPath);
+	if (MenuItemBlock.IsValid())
 	{
-		ChildText = FindChildWidgetWithText(MenuItemBlock);
+		FText ChildText = FindChildWidgetWithText(MenuItemBlock);
 		if (!ChildText.IsEmpty())
 		{
 			SlateQueryContext.ItemName = ChildText;
 			SlateQueryContext.ItemDescriptor = LOCTEXT("ItemDescriptor_Menu", "menu item");
-			SlateQueryContext.LastPickedWidget = MenuItemBlock.ToSharedPtr();
+			SlateQueryContext.LastPickedWidget = MenuItemBlock;
 			SlateQueryContext.bIsUIWidget = true;
+			return true;
 		}
 	}
+	return false;
+}
 
-	// Is it a details panel property? Look up for an SDetailSingleItemRow [could be others?] then down through the SPropertyNameWidget.
-	TSharedRef<SWidget> PropertyRow = FindClosestWidgetOfType(WidgetPath, "SDetailSingleItemRow");
-	if (PropertyRow->GetType() != "SNullWidgetContent")
+// Return true and fill in SlateQueryContext if queried item is a Details panel property.
+static bool ItemIsDetailsProperty(const FWidgetPath& WidgetPath, FAIAssistantSlateQueryContext& SlateQueryContext)
+{
+	// Look up for an SDetailSingleItemRow [could be others?] then down through the SPropertyNameWidget.
+	TSharedPtr<SWidget> PropertyRow = FindClosestWidgetPtrOfType(WidgetPath, "SDetailSingleItemRow");
+	if (PropertyRow.IsValid())
 	{
-		SlateQueryContext.LastPickedWidget = PropertyRow.ToSharedPtr();
+		FText ChildText;
+
+		SlateQueryContext.LastPickedWidget = PropertyRow;
 
-		TSharedRef<SWidget> PropertyNameWidget = FindChildWidgetOfType(PropertyRow, "SPropertyNameWidget");
-		if (PropertyNameWidget->GetType() != "SNullWidgetContent")
+		TSharedPtr<SWidget> PropertyNameWidget = FindChildWidgetOfType(PropertyRow, "SPropertyNameWidget");
+		if (PropertyNameWidget.IsValid())
 		{
 			ChildText = FindChildWidgetWithText(PropertyNameWidget);
 			if (!ChildText.IsEmpty())
@@ -560,29 +584,39 @@ static bool FindItemName(const FWidgetPath& WidgetPath, FAIAssistantSlateQueryCo
 				}
 			}
 		}
+		return true;
 	}
+	return false;
+}
 
-	// Is it a breadcrumb trail button?
-	TSharedRef<SWidget> BreadcrumbTrail = FindClosestWidgetOfType(WidgetPath, "SBreadcrumbTrail<FNavigationCrumb>");
-	if (BreadcrumbTrail->GetType() != "SNullWidgetContent")
+// Return true and fill in SlateQueryContext if queried item is a navigation breadcrumb button.
+static bool ItemIsBreadcrumbTrailButton(const FWidgetPath& WidgetPath, FAIAssistantSlateQueryContext& SlateQueryContext)
+{
+	TSharedPtr<SWidget> BreadcrumbTrail = FindClosestWidgetPtrOfType(WidgetPath, "SBreadcrumbTrail<FNavigationCrumb>");
+	TSharedPtr<SWidget> BreadcrumbButton = BreadcrumbTrail ? FindClosestWidgetPtrOfType(WidgetPath, "SButton") : nullptr;
+	FText ChildText = FindChildWidgetWithText(BreadcrumbButton);
+	if (!ChildText.IsEmpty())
 	{
-		TSharedRef<SWidget> BreadcrumbButton = FindClosestWidgetOfType(WidgetPath, "SButton");
-		ChildText = FindChildWidgetWithText(BreadcrumbButton);
-		if (!ChildText.IsEmpty())
-		{
-			SlateQueryContext.ItemName = ChildText;
-			SlateQueryContext.ItemDescriptor = LOCTEXT("ItemDescriptor_Breadcrumb", "navigation breadcrumb");
-			SlateQueryContext.LastPickedWidget = BreadcrumbButton.ToSharedPtr();
-			SlateQueryContext.bIsUIWidget = true;
-		}
+		SlateQueryContext.ItemName = ChildText;
+		SlateQueryContext.ItemDescriptor = LOCTEXT("ItemDescriptor_Breadcrumb", "navigation breadcrumb");
+		SlateQueryContext.LastPickedWidget = BreadcrumbButton;
+		SlateQueryContext.bIsUIWidget = true;
+		return true;
 	}
+	return false;
+}
+
+// Return true and fill in SlateQueryContext if queried item is an asset item.
+static bool ItemIsAsset(const FWidgetPath& WidgetPath, FAIAssistantSlateQueryContext& SlateQueryContext)
+{
+	FText ChildText;
 
 	// Is it an asset tile item?
-	TSharedRef<SWidget> AssetTileItem = FindClosestWidgetOfType(WidgetPath, "SAssetTileItem");
-	if (AssetTileItem->GetType() != "SNullWidgetContent")
+	TSharedPtr<SWidget> AssetTileItem = FindClosestWidgetPtrOfType(WidgetPath, "SAssetTileItem");
+	if (AssetTileItem.IsValid())
 	{
-		TSharedRef<SWidget> AssetThumbnail = FindChildWidgetOfType(AssetTileItem, "SAssetThumbnail");
-		if (AssetThumbnail->GetType() != "SNullWidgetContent")
+		TSharedPtr<SWidget> AssetThumbnail = FindChildWidgetOfType(AssetTileItem, "SAssetThumbnail");
+		if (AssetThumbnail.IsValid())
 		{
 			ChildText = FindChildWidgetWithText(AssetThumbnail); // This returns TYPE of asset instead of name
 			if (!ChildText.IsEmpty())
@@ -597,7 +631,7 @@ static bool FindItemName(const FWidgetPath& WidgetPath, FAIAssistantSlateQueryCo
 				// We can use the tooltip for the asset's name.
 				SlateQueryContext.ItemName = SlateQueryContext.CurrentToolTipText;
 			}
-			SlateQueryContext.LastPickedWidget = AssetThumbnail.ToSharedPtr();
+			SlateQueryContext.LastPickedWidget = AssetThumbnail;
 			SlateQueryContext.bIsObject = true;
 		}
 		else
@@ -607,94 +641,148 @@ static bool FindItemName(const FWidgetPath& WidgetPath, FAIAssistantSlateQueryCo
 			{
 				SlateQueryContext.ItemName = ChildText;
 				SlateQueryContext.ItemDescriptor = LOCTEXT("ItemDescriptor_Folder", "asset folder");
-				SlateQueryContext.LastPickedWidget = AssetTileItem.ToSharedPtr();
+				SlateQueryContext.LastPickedWidget = AssetTileItem;
 				SlateQueryContext.bIsObject = true;
 			}
 		}
+		return true;
 	}
 
 	// Is it an asset tree item?
-	TSharedRef<SWidget> AssetTreeItem = FindClosestWidgetOfType(WidgetPath, "SAssetTreeItem");
-	if (AssetTreeItem->GetType() != "SNullWidgetContent")
+	TSharedPtr<SWidget> AssetTreeItem = FindClosestWidgetPtrOfType(WidgetPath, "SAssetTreeItem");
+	if (AssetTreeItem.IsValid())
 	{
 		ChildText = FindChildWidgetWithText(AssetTreeItem);
 		if (!ChildText.IsEmpty())
 		{
 			SlateQueryContext.ItemName = ChildText;
 			SlateQueryContext.ItemDescriptor = LOCTEXT("ItemDescriptor_Folder", "asset folder");
-			SlateQueryContext.LastPickedWidget = AssetTreeItem.ToSharedPtr();
+			SlateQueryContext.LastPickedWidget = AssetTreeItem;
 			SlateQueryContext.bIsObject = true;
+			return true;
 		}
 	}
+	return false;
+}
+
+// Return true and fill in SlateQueryContext if queried item is in Outliner or Subobject Instance Editor.
+static bool ItemIsOutlinerActorOrSubobject(const FWidgetPath& WidgetPath, FAIAssistantSlateQueryContext& SlateQueryContext)
+{
+	FText ChildText;
 
 	// Is it an actor in the Outliner?
-	TSharedRef<SWidget> Outliner = FindClosestWidgetOfType(WidgetPath, "SSceneOutliner");
-	if (Outliner->GetType() != "SNullWidgetContent")
+	TSharedPtr<SWidget> Outliner = FindClosestWidgetPtrOfType(WidgetPath, "SSceneOutliner");
+	TSharedPtr<SWidget> ActorTreeLabel = Outliner ? FindClosestWidgetPtrOfType(WidgetPath, "SActorTreeLabel") : nullptr;
+	TSharedPtr<SWidget> OutlinerRow = Outliner ? FindClosestWidgetPtrOfType(WidgetPath, "SSceneOutlinerTreeRow") : nullptr;
+
+	if (!ActorTreeLabel.IsValid() && OutlinerRow.IsValid())
 	{
-		TSharedRef<SWidget> ActorTreeLabel = FindClosestWidgetOfType(WidgetPath, "SActorTreeLabel");
-		if (ActorTreeLabel->GetType() != "SNullWidgetContent")
-		{
-			ChildText = FindChildWidgetWithText(ActorTreeLabel);
-			if (!ChildText.IsEmpty())
-			{
-				SlateQueryContext.ItemName = ChildText;
-				SlateQueryContext.ItemDescriptor = LOCTEXT("ItemDescriptor_OutlinerActor", "actor");
-				SlateQueryContext.LastPickedWidget = ActorTreeLabel.ToSharedPtr();
-				SlateQueryContext.bIsObject = true;
-			}
-		}
-		else
-		{
-			TSharedRef<SWidget> OutlinerRow = FindClosestWidgetOfType(WidgetPath, "SSceneOutlinerTreeRow");
-			if (OutlinerRow->GetType() != "SNullWidgetContent")
-			{
-				SlateQueryContext.LastPickedWidget = OutlinerRow.ToSharedPtr();
-				SlateQueryContext.bIsObject = true;
-				ActorTreeLabel = FindChildWidgetOfType(OutlinerRow, "SActorTreeLabel");
-				if (ActorTreeLabel->GetType() != "SNullWidgetContent")
-				{
-					ChildText = FindChildWidgetWithText(ActorTreeLabel);
-					if (!ChildText.IsEmpty())
-					{
-						SlateQueryContext.ItemName = ChildText;
-						SlateQueryContext.ItemDescriptor = LOCTEXT("ItemDescriptor_OutlinerActor", "actor");
-					}
-				}
-			}
-		}
+		ActorTreeLabel = FindChildWidgetOfType(OutlinerRow, "SActorTreeLabel");
 	}
-	// Is it an object in the Subobject Instance editor?
-	TSharedRef<SWidget> SubobjectInstanceEditor = FindClosestWidgetOfType(WidgetPath, "SSubobjectInstanceEditor");
-	if (SubobjectInstanceEditor->GetType() != "SNullWidgetContent")
+
+	if (ActorTreeLabel.IsValid())
 	{
-		TSharedRef<SWidget> SubObjectRow = FindClosestWidgetOfType(WidgetPath, "SSubobject_RowWidget");
-		if (SubObjectRow->GetType() != "SNullWidgetContent")
+		ChildText = FindChildWidgetWithText(ActorTreeLabel);
+		if (!ChildText.IsEmpty())
 		{
-			SlateQueryContext.LastPickedWidget = SubObjectRow.ToSharedPtr();
+			SlateQueryContext.ItemName = ChildText;
+			SlateQueryContext.ItemDescriptor = LOCTEXT("ItemDescriptor_OutlinerActor", "actor");
+			SlateQueryContext.LastPickedWidget = ActorTreeLabel;
 			SlateQueryContext.bIsObject = true;
-			TSharedRef<SWidget> NameWidget = FindChildWidgetOfType(SubObjectRow, "SInlineEditableTextBlock");
-			if (NameWidget->GetType() != "SNullWidgetContent")
-			{
-				ChildText = FindChildWidgetWithText(NameWidget);
-				SlateQueryContext.ItemName = ChildText;
-				SlateQueryContext.ItemDescriptor = LOCTEXT("ItemDescriptor_SubobjectInstance", "instance");
-			}
+			return true;
 		}
 	}
 
-	return true;
+	// Is it an object in the Subobject Instance editor?
+	TSharedPtr<SWidget> SubobjectInstanceEditor = FindClosestWidgetPtrOfType(WidgetPath, "SSubobjectInstanceEditor");
+	TSharedPtr<SWidget> SubObjectRow = SubobjectInstanceEditor ? FindClosestWidgetPtrOfType(WidgetPath, "SSubobject_RowWidget") : nullptr;
+	TSharedPtr<SWidget> NameWidget = SubObjectRow ? FindChildWidgetOfType(SubObjectRow, "SInlineEditableTextBlock") : nullptr;
+
+	if (NameWidget.IsValid())
+	{
+		ChildText = FindChildWidgetWithText(NameWidget);
+		SlateQueryContext.ItemName = ChildText;
+		SlateQueryContext.ItemDescriptor = LOCTEXT("ItemDescriptor_SubobjectInstance", "instance");
+		SlateQueryContext.LastPickedWidget = SubObjectRow;
+		SlateQueryContext.bIsObject = true;
+		return true;
+	}
+
+	return false;
+}
+
+// Return true and fill in SlateQueryContext if queried item is a Details panel filter checkbox.
+static bool ItemIsDetailPanelFilterCheckbox(const FWidgetPath& WidgetPath, FAIAssistantSlateQueryContext& SlateQueryContext)
+{
+	TSharedPtr<SWidget> ActorDetails = FindClosestWidgetPtrOfType(WidgetPath, "SActorDetails");
+	if (!ActorDetails)
+	{
+		return false;
+	}
+
+	TSharedPtr<SWidget> DetailsView = FindClosestWidgetPtrOfType(WidgetPath, "SDetailsView");
+	TSharedPtr<SWidget> DetailsNameArea = FindClosestWidgetPtrOfType(WidgetPath, "SDetailsNameArea");
+	TSharedPtr<SWidget> SubobjectInstanceEditor = FindClosestWidgetPtrOfType(WidgetPath, "SSubobjectInstanceEditor");
+	TSharedPtr<SWidget> WrapBox = FindClosestWidgetPtrOfType(WidgetPath, "SWrapBox");
+	TSharedPtr<SWidget> CheckBox = FindClosestWidgetPtrOfType(WidgetPath, "SCheckBox");
+
+	if (!DetailsView && !DetailsNameArea && !SubobjectInstanceEditor && WrapBox && CheckBox)
+	{
+		return FillButtonQueryContext(CheckBox, SlateQueryContext,
+			LOCTEXT("ItemDescriptor_PropertyFilter", "property filter"));
+	}
+	return false;
 }
 
+// Return true and fill in SlateQueryContext with item information if valid item information is found.
+static bool FindItemName(const FWidgetPath& WidgetPath, FAIAssistantSlateQueryContext& SlateQueryContext)
+{
+	SlateQueryContext.ItemDescriptor = LOCTEXT("ItemDescriptor_Generic", "control");
 
+	if (!WidgetPath.IsValid())
+	{
+		return false;
+	}
+
+	if (// Does the widget path end up in a Viewport?
+		ItemIsInViewport(WidgetPath, SlateQueryContext) ||
+		// Is it a graph node?
+		ItemIsGraphNode(WidgetPath, SlateQueryContext) ||
+		// Is it a search/filter field?
+		ItemIsSearchBox(WidgetPath, SlateQueryContext) ||
+		// Is it a console input box?
+		ItemIsConsoleInputBox(WidgetPath, SlateQueryContext) ||
+		// Is it a details panel property?
+		ItemIsDetailsProperty(WidgetPath, SlateQueryContext) ||
+		// Is it a breadcrumb trail button?
+		ItemIsBreadcrumbTrailButton(WidgetPath, SlateQueryContext) ||
+		// Is it an asset item?
+		ItemIsAsset(WidgetPath, SlateQueryContext) ||
+		// Is it an Actor in Outliner, or an instance or component in the SubobjectInstanceEditor?
+		ItemIsOutlinerActorOrSubobject(WidgetPath, SlateQueryContext) ||
+		// Is it one of the Details panel checkboxes that filter the detail properties viewed?
+		ItemIsDetailPanelFilterCheckbox(WidgetPath, SlateQueryContext) ||
+		// Is it a menu item?
+		ItemIsMenuItem(WidgetPath, SlateQueryContext) ||
+		// Is it a button?
+		ItemIsButton(WidgetPath, SlateQueryContext))
+	{
+		return true;
+	}
+
+	return false;
+}
+
+// Return true and fill in SlateQueryContext if valid Tab or Panel information for queried item is found.
 static bool FindTabName(const FWidgetPath& WidgetPath, FAIAssistantSlateQueryContext& SlateQueryContext)
 {
 	bool bFoundTabName = false;
 
-	TSharedRef<SWidget> TabStack = FindClosestWidgetOfType(WidgetPath, "SDockingTabStack");
-	if (TabStack->GetType() != "SNullWidgetContent")
+	TSharedPtr<SWidget> TabStack = FindClosestWidgetPtrOfType(WidgetPath, "SDockingTabStack");
+	if (TabStack.IsValid())
 	{
 		TArray<TSharedRef<SWidget>> DockTabs;
-		FindChildWidgetsOfType(DockTabs, TabStack.ToSharedPtr(), "SDockTab");
+		FindChildWidgetsOfType(DockTabs, TabStack, "SDockTab");
 		for (auto& ThisTab : DockTabs)
 		{
 			TSharedRef<SDockTab> ThisTabCast = StaticCastSharedRef<SDockTab>(ThisTab);
@@ -704,7 +792,7 @@ static bool FindTabName(const FWidgetPath& WidgetPath, FAIAssistantSlateQueryCon
 			}
 			if (ThisTabCast->IsForeground())
 			{
-				SlateQueryContext.LastPickedWidget = TabStack.ToSharedPtr();
+				SlateQueryContext.LastPickedWidget = TabStack;
 				SlateQueryContext.StructuredContext.TabName = ThisTabCast->GetTabLabel().ToString();
 				FFormatNamedArguments Args;
 				Args.Add(TEXT("PanelName"), ThisTabCast->GetTabLabel());
@@ -717,14 +805,13 @@ static bool FindTabName(const FWidgetPath& WidgetPath, FAIAssistantSlateQueryCon
 	return bFoundTabName;
 }
 
-
+// If we're in a details panel, include the class of actor or object that we're editing in SlateQueryContext.
 static void GenerateDetailsViewContext(const FWidgetPath& WidgetPath, FAIAssistantSlateQueryContext& SlateQueryContext)
 {
-	// if we're in a details panel, include the class of actor or object that we're editing
-	const TSharedRef<SWidget> DetailsView = FindClosestWidgetOfType(WidgetPath, "SDetailsView");
-	if (DetailsView->GetType() != "SNullWidgetContent")
+	const TSharedPtr<SWidget> DetailsView = FindClosestWidgetPtrOfType(WidgetPath, "SDetailsView");
+	if (DetailsView.IsValid())
 	{
-		const TSharedRef<IDetailsView> DetailsViewCast = StaticCastSharedRef<IDetailsView>(DetailsView);
+		const TSharedPtr<IDetailsView> DetailsViewCast = StaticCastSharedPtr<IDetailsView>(DetailsView);
 		if (DetailsViewCast->GetSelectedActors().Num() > 0)
 		{
 			FString ActorName;
@@ -740,35 +827,17 @@ static void GenerateDetailsViewContext(const FWidgetPath& WidgetPath, FAIAssista
 	}
 }
 
-
-static bool FindEditorName(const FWidgetPath& WidgetPath, FAIAssistantSlateQueryContext& SlateQueryContext)
+// Return true and fill in SlateQueryContext with item information if valid editor information is found.
+static bool IsInEditor(const FWidgetPath& WidgetPath, FAIAssistantSlateQueryContext& SlateQueryContext)
 {
-	FText OutName = FText();
-
-	if (!WidgetPath.IsValid())
-	{
-		SlateQueryContext.FormattedWindowName = OutName;
-		return false;
-	}
-
-	// is it in the status bar?
-	TSharedRef<SWidget> StatusBar = FindClosestWidgetOfType(WidgetPath, "SStatusBar");
-	if (StatusBar->GetType() != "SNullWidgetContent")
-	{
-		SlateQueryContext.LastPickedWidget = StatusBar.ToSharedPtr();
-		SlateQueryContext.StructuredContext.TabName = TEXT("StatusBar");
-		SlateQueryContext.FormattedWindowName = LOCTEXT("TabName_StatusBar", " the Status Bar");
-		return true;
-	}
-
-	TSet<FEditorModeID> BlockedModeIDs = TSet<FEditorModeID>({
+	static const TSet<FEditorModeID> BlockedModeIDs = TSet<FEditorModeID>({
 		FBuiltinEditorModes::EM_Default,
 		"EditMode.SubTrackEditMode"
 		});
 
 	// are we in the level editor?
-	TSharedRef<SWidget> LevelEditor = FindFirstWidgetOfType(WidgetPath, "SLevelEditor");
-	if (LevelEditor->GetType() != "SNullWidgetContent")
+	TSharedPtr<SWidget> LevelEditor = FindFirstWidgetPtrOfType(WidgetPath, "SLevelEditor");
+	if (LevelEditor)
 	{
 		SlateQueryContext.StructuredContext.EditorName = TEXT("LevelEditor");
 		FText ModeString;
@@ -779,6 +848,18 @@ static bool FindEditorName(const FWidgetPath& WidgetPath, FAIAssistantSlateQuery
 			TArray<FEditorModeInfo> ModeInfos = AssetEditorSubsystem->GetEditorModeInfoOrderedByPriority();
 			for (const FEditorModeInfo& ModeInfo : ModeInfos)
 			{
+				if (UEdMode* EdMode = ModeManager->GetActiveScriptableMode(ModeInfo.ID))
+				{
+					TWeakPtr<FModeToolkit> ModeToolkit = EdMode->GetToolkit();
+					if (ModeToolkit.IsValid())
+					{
+						FText ToolDisplayName = ModeToolkit.Pin()->GetActiveToolDisplayName();
+						if (!ToolDisplayName.IsEmpty())
+						{
+							SlateQueryContext.StructuredContext.ActiveToolName = ToolDisplayName.ToString();
+						}
+					}
+				}
 				if (ModeString.IsEmpty() && ModeManager->IsModeActive(ModeInfo.ID) && !BlockedModeIDs.Contains(ModeInfo.ID))
 				{
 					FFormatNamedArguments Args;
@@ -788,27 +869,67 @@ static bool FindEditorName(const FWidgetPath& WidgetPath, FAIAssistantSlateQuery
 				}
 			}
 		}
-		SlateQueryContext.LastPickedWidget = LevelEditor.ToSharedPtr();
+		SlateQueryContext.LastPickedWidget = LevelEditor;
 		if (ModeString.IsEmpty())
 		{
 			SlateQueryContext.FormattedWindowName = LOCTEXT("EditorName_LevelEditor", " the Level Editor");
-			return true;
 		}
 		else
 		{
 			FFormatNamedArguments Args;
 			Args.Add(TEXT("WithModeActive"), ModeString);
 			SlateQueryContext.FormattedWindowName = FText::Format(LOCTEXT("EditorName_LevelEditorAndMode", " the Level Editor {WithModeActive}"), Args);
-			return true;
 		}
+
+		// Get selected actors, if any.
+		bool bActorIsSelected = false;
+		UEditorActorSubsystem* EditorActorSubsystem = GEditor->GetEditorSubsystem<UEditorActorSubsystem>();
+		// Get selected actors
+		TArray<AActor*> SelectedActors = EditorActorSubsystem->GetSelectedLevelActors();
+		for (AActor* Actor : SelectedActors)
+		{
+			if (Actor)
+			{
+				SlateQueryContext.SelectedActorName = Actor->GetName();
+				SlateQueryContext.SelectedActorClass = Actor->GetClass()->GetName();
+				bActorIsSelected = true;
+				break;
+			}
+		}
+		if (!bActorIsSelected)
+		{
+			FViewport* ActiveViewport = GEditor->GetActiveViewport();
+			if (ActiveViewport)
+			{
+				FIntPoint MousePos;
+				ActiveViewport->GetMousePos(MousePos);
+				HHitProxy* HitProxy = ActiveViewport->GetHitProxy(MousePos.X, MousePos.Y);
+				if (HitProxy)
+				{
+					if (HitProxy->IsA(HActor::StaticGetType()))
+					{
+						HActor* TargetProxy = static_cast<HActor*>(HitProxy);
+						AActor* TargetActor = TargetProxy->Actor;
+						// Double check for not null
+						if (TargetActor)
+						{
+							SlateQueryContext.SelectedActorName = TargetActor->GetName();
+							SlateQueryContext.SelectedActorClass = TargetActor->GetClass()->GetName();
+						}
+					}
+				}
+			}
+		}
+
+		return true;
 	}
 
 	// if not, are we in a different asset editor?
-	TSharedRef<SWidget> StandaloneAssetEditorHost = FindFirstWidgetOfType(WidgetPath, "SStandaloneAssetEditorToolkitHost");
-	if (StandaloneAssetEditorHost->GetType() != "SNullWidgetContent")
+	TSharedPtr<SWidget> StandaloneAssetEditorHost = FindFirstWidgetPtrOfType(WidgetPath, "SStandaloneAssetEditorToolkitHost");
+	if (StandaloneAssetEditorHost)
 	{
 		TArray<TSharedRef<SWidget>> TabStacks;
-		FindChildWidgetsOfType(TabStacks, StandaloneAssetEditorHost.ToSharedPtr(), "SDockingTabStack");
+		FindChildWidgetsOfType(TabStacks, StandaloneAssetEditorHost, "SDockingTabStack");
 		for (auto& ThisTabStack : TabStacks)
 		{
 			TArray<TSharedRef<SWidget>> DockTabs;
@@ -834,7 +955,12 @@ static bool FindEditorName(const FWidgetPath& WidgetPath, FAIAssistantSlateQuery
 								SlateQueryContext.LastPickedWidget = ThisTabStack.ToSharedPtr();
 								FFormatNamedArguments Args;
 								Args.Add(TEXT("AssetEditorName"), FText::FromName(ThisEditor->GetEditorName()));
-								SlateQueryContext.FormattedWindowName = FText::Format(LOCTEXT("EditorName_AssetEditorFormat", " the {AssetEditorName} editor"), Args);
+								SlateQueryContext.FormattedWindowName = FText::Format(
+									LOCTEXT("EditorName_AssetEditorFormat", " the {AssetEditorName} editor"), Args);
+
+								SlateQueryContext.SelectedActorName = ThisEditedAsset->GetName();
+								SlateQueryContext.SelectedActorClass = ThisEditedAsset->GetClass()->GetName();
+
 								return true;
 							}
 						}
@@ -843,39 +969,81 @@ static bool FindEditorName(const FWidgetPath& WidgetPath, FAIAssistantSlateQuery
 			}
 		}
 	}
+	return false;
+}
 
-	// Is it in a drawer overlay?
-	TSharedRef<SWidget> DrawerOverlay = FindFirstWidgetOfType(WidgetPath, "SDrawerOverlay");
-	if (DrawerOverlay->GetType() != "SNullWidgetContent")
+// Return true and fill in SlateQueryContext with item information if queried item is in a drawer and valid
+// window and tab information is found.
+static bool IsInDrawer(const FWidgetPath& WidgetPath, FAIAssistantSlateQueryContext& SlateQueryContext)
+{
+	TSharedPtr<SWidget> DrawerOverlay = FindFirstWidgetPtrOfType(WidgetPath, "SDrawerOverlay");
+	if (DrawerOverlay.IsValid())
 	{
 		FText DrawerName;
-
-		TSharedRef<SWidget> DrawerWidget = FindChildWidgetOfType(DrawerOverlay, "SContentBrowser");
-		if (DrawerWidget->GetType() != "SNullWidgetContent")
+		TSharedPtr<SWidget> DrawerWidget = FindChildWidgetOfType(DrawerOverlay, "SContentBrowser");
+		if (!DrawerWidget)
 		{
-			DrawerName = LOCTEXT("DrawerName_ContentBrowser", "ContentBrowser");
+			DrawerWidget = FindChildWidgetOfType(DrawerOverlay, "SOutputLog");
 		}
-		else
+
+		if (DrawerWidget)
 		{
-			DrawerWidget = FindChildWidgetOfType(DrawerOverlay, "SOutputLog");
-			if (DrawerWidget->GetType() != "SNullWidgetContent")
+			if (DrawerWidget->GetType() == "SContentBrowser")
+			{
+				DrawerName = LOCTEXT("DrawerName_ContentBrowser", "ContentBrowser");
+			}
+			else if (DrawerWidget->GetType() == "SOutputLog")
 			{
 				DrawerName = LOCTEXT("DrawerName_OutputLog", "OutputLog");
 			}
-		}
-		// Are there other items that are commonly in drawers?
+			// Are there other items that are commonly in drawers?
 
-		if (!DrawerName.IsEmpty())
-		{
-			SlateQueryContext.StructuredContext.WindowName = DrawerOverlay->GetType().ToString();
-			SlateQueryContext.StructuredContext.TabName = DrawerName.ToString();
-			SlateQueryContext.LastPickedWidget = DrawerWidget.ToSharedPtr();
-			FFormatNamedArguments Args;
-			Args.Add(TEXT("DrawerName"), DrawerName);
-			SlateQueryContext.FormattedWindowName = FText::Format(LOCTEXT("DrawerNameFormat", "the {DrawerName} drawer"), Args);
-			return true;
+			if (!DrawerName.IsEmpty())
+			{
+				SlateQueryContext.StructuredContext.WindowName = DrawerOverlay->GetType().ToString();
+				SlateQueryContext.StructuredContext.TabName = DrawerName.ToString();
+				SlateQueryContext.LastPickedWidget = DrawerWidget;
+				FFormatNamedArguments Args;
+				Args.Add(TEXT("DrawerName"), DrawerName);
+				SlateQueryContext.FormattedWindowName = FText::Format(LOCTEXT("DrawerNameFormat", "the {DrawerName} drawer"), Args);
+				return true;
+			}
 		}
 	}
+	return false;
+}
+
+// Return true and fill in SlateQueryContext with item information if valid editor or window information is found.
+static bool FindEditorName(const FWidgetPath& WidgetPath, FAIAssistantSlateQueryContext& SlateQueryContext)
+{
+	FText OutName = FText();
+
+	if (!WidgetPath.IsValid())
+	{
+		SlateQueryContext.FormattedWindowName = OutName;
+		return false;
+	}
+
+	// is it in the status bar?
+	TSharedPtr<SWidget> StatusBar = FindClosestWidgetPtrOfType(WidgetPath, "SStatusBar");
+	if (StatusBar.IsValid())
+	{
+		SlateQueryContext.LastPickedWidget = StatusBar;
+		SlateQueryContext.StructuredContext.TabName = TEXT("StatusBar");
+		SlateQueryContext.FormattedWindowName = LOCTEXT("TabName_StatusBar", " the Status Bar");
+		return true;
+	}
+
+	if (IsInEditor(WidgetPath, SlateQueryContext))
+	{
+		return true;
+	}
+
+	// Is it in a drawer overlay?
+	if (IsInDrawer(WidgetPath, SlateQueryContext))
+	{
+		return true;
+	}
 
 	TSharedRef<SWindow> WidgetWindow = WidgetPath.GetWindow();
 	// is it the main window that hosts the level editor?
@@ -906,7 +1074,7 @@ static bool FindEditorName(const FWidgetPath& WidgetPath, FAIAssistantSlateQuery
 	return false;
 }
 
-
+// Fill in AccumulatedSlateQueryContext with information on text under cursor if found.
 static void FindTextUnderCursor(const FWidgetPath& WidgetPath, FAIAssistantSlateQueryContext& AccumulatedSlateQueryContext)
 {
 	FText OutText;
@@ -1018,7 +1186,11 @@ void UE::AIAssistant::SlateQuerier::QueryAIAssistantAboutSlateWidget(const FWidg
 		}
 
 		FText NameToQuery;
-		if (FormattedItemName.IsEmpty() && !SlateQueryContext.TextUnderCursor.IsEmpty())
+		if (SlateQueryContext.bIsInAssistantPanel)
+		{
+			NameToQuery = LOCTEXT("ItemName_Assistant", "the AI Assistant");
+		}
+		else if (FormattedItemName.IsEmpty() && !SlateQueryContext.TextUnderCursor.IsEmpty())
 		{
 			NameToQuery = SlateQueryContext.TextUnderCursor;
 		}
@@ -1030,19 +1202,20 @@ void UE::AIAssistant::SlateQuerier::QueryAIAssistantAboutSlateWidget(const FWidg
 		static const FText QueryFormatWidget = LOCTEXT("QueryFormat1", "I would like to know what {NameToQuery} does.");
 		static const FText QueryFormatGeneral = LOCTEXT("QueryFormat2", "I would like to know what \"{NameToQuery}\" means.");
 		static const FText QueryFormatObject = LOCTEXT("QueryFormat3", "I would like to know about {NameToQuery}.");
+		static const FText QueryFormatAssistant = LOCTEXT("QueryFormat4", "Tell me about yourself.");
 		static const FText QueryInstructionsWidgetFormat = LOCTEXT("QueryInstructionsFormat1",
 			"Provide an easily readable, informative, accurate answer that describes what the widget does in the Unreal Editor UI. {QueryAdditionalInstructions} {StyleInstructions}");
 		static const FText QueryInstructionsObjectFormat = LOCTEXT("QueryInstructionsFormat2",
 			"Provide an easily readable, informative, accurate answer that describes what I should know about it in the Unreal Editor UI. {QueryAdditionalInstructions} {StyleInstructions}");
+		static const FText QueryInstructionsAssistantFormat = LOCTEXT("QueryInstructionsFormat3",
+			"Provide an easily readable, informative, accurate answer that describes what I should know about yourself and what you can do for me. {StyleInstructions}");
 		static const FText QueryAdditionalInstructionsWidget = LOCTEXT("QueryAdditionalInstructions1",
 			"Explain what happens if I select or click on it.");
 		static const FText QueryAdditionalInstructionsObject = LOCTEXT("QueryAdditionalInstructions2",
 			"Describe what the object is or what it does.");
 		static const FText QueryAdditionalInstructionsGeneral = LOCTEXT("QueryAdditionalInstructions3",
 			"If the object is clickable, like a button, then explain what happens if I click on it. If the object is something you select, like an asset or a blueprint node, describe what the object is or what it does.");
-		static const FText StyleInstructionsWidget = LOCTEXT("StyleInstructions1",
-			"Use bold text to emphasize Unreal Engine specific terms. Start with a quick overview paragraph, then add key points limited to one or two formatted sections with headers and bullet points. Don't include a summary at the end.");
-		static const FText StyleInstructionsObject = LOCTEXT("StyleInstructions2",
+		static const FText StyleInstructions = LOCTEXT("StyleInstructions1",
 			"Use bold text to emphasize Unreal Engine specific terms. Start with a quick overview paragraph, then add key points limited to one or two formatted sections with headers and bullet points. Don't include a summary at the end.");
 
 		if (!NameToQuery.IsEmpty())
@@ -1051,25 +1224,31 @@ void UE::AIAssistant::SlateQuerier::QueryAIAssistantAboutSlateWidget(const FWidg
 			QueryArgs.Add(TEXT("NameToQuery"), NameToQuery);
 			FFormatNamedArguments InstructionArgs;
 
-			if (SlateQueryContext.bIsObject)
+			if (SlateQueryContext.bIsInAssistantPanel)
+			{
+				SlateQueryContext.GeneratedQuery = FText::Format(QueryFormatAssistant, QueryArgs);
+				InstructionArgs.Add(TEXT("StyleInstructions"), StyleInstructions);
+				SlateQueryContext.GeneratedQueryInstructions = FText::Format(QueryInstructionsAssistantFormat, InstructionArgs);
+			}
+			else if (SlateQueryContext.bIsObject)
 			{
 				SlateQueryContext.GeneratedQuery = FText::Format(QueryFormatObject, QueryArgs);
 				InstructionArgs.Add(TEXT("QueryAdditionalInstructions"), QueryAdditionalInstructionsObject);
-				InstructionArgs.Add(TEXT("StyleInstructions"), StyleInstructionsObject);
+				InstructionArgs.Add(TEXT("StyleInstructions"), StyleInstructions);
 				SlateQueryContext.GeneratedQueryInstructions = FText::Format(QueryInstructionsObjectFormat, InstructionArgs);
 			}
 			else if (SlateQueryContext.bIsUIWidget)
 			{
 				SlateQueryContext.GeneratedQuery = FText::Format(QueryFormatWidget, QueryArgs);
 				InstructionArgs.Add(TEXT("QueryAdditionalInstructions"), QueryAdditionalInstructionsWidget);
-				InstructionArgs.Add(TEXT("StyleInstructions"), StyleInstructionsWidget);
+				InstructionArgs.Add(TEXT("StyleInstructions"), StyleInstructions);
 				SlateQueryContext.GeneratedQueryInstructions = FText::Format(QueryInstructionsWidgetFormat, InstructionArgs);
 			}
 			else
 			{
 				SlateQueryContext.GeneratedQuery = FText::Format(QueryFormatWidget, QueryArgs);
 				InstructionArgs.Add(TEXT("QueryAdditionalInstructions"), QueryAdditionalInstructionsGeneral);
-				InstructionArgs.Add(TEXT("StyleInstructions"), StyleInstructionsWidget);
+				InstructionArgs.Add(TEXT("StyleInstructions"), StyleInstructions);
 				SlateQueryContext.GeneratedQueryInstructions = FText::Format(QueryInstructionsWidgetFormat, InstructionArgs);
 			}
 		}
@@ -1085,7 +1264,6 @@ void UE::AIAssistant::SlateQuerier::QueryAIAssistantAboutSlateWidget(const FWidg
 	{
 		SlateQueryContext.StructuredContext.ToolTipText = SlateQueryContext.CurrentToolTipText.ToString();
 
-		GenerateModeToolContext(InWidgetPath, SlateQueryContext);
 		GenerateDetailsViewContext(InWidgetPath, SlateQueryContext);
 
 		if (!SlateQueryContext.ItemName.IsEmpty())
@@ -1108,7 +1286,7 @@ void UE::AIAssistant::SlateQuerier::QueryAIAssistantAboutSlateWidget(const FWidg
 			SlateQueryContext.StructuredContext.WidgetType = SlateQueryContext.LastPickedWidget->GetType().ToString();
 		}
 
-		if (SlateQueryContext.ItemName.IsEmpty())
+		if (SlateQueryContext.ItemName.IsEmpty() && !SlateQueryContext.bIsInAssistantPanel)
 		{
 			SlateQueryContext.StructuredContext.WidgetPath = TArray<FString>();
 			// Add simplified widget path
@@ -1130,10 +1308,6 @@ void UE::AIAssistant::SlateQuerier::QueryAIAssistantAboutSlateWidget(const FWidg
 		}
 	}
 
-  // DEBUG
-  FString StructuredContextOutput = SlateQueryContext.StructuredContext.ToJson();
-  UE_LOG(LogAIAssistant, Display, TEXT("context json: %s"), *StructuredContextOutput);
-
 	// OPTIONAL - We're done using these members. We can clear them to reduce size as they're
 	// not used below.
 	SlateQueryContext.LastPickedWidget.Reset();
