[AI Assistant] Add ability to reset the conversation ready executor.

This allows the conversation ready executor to reset to it initial state when an error occurs without
having to clear the queue (e.g reconstruct the object).

diff --git a/Engine/Plugins/Experimental/AIAssistant/Source/AIAssistant/Private/AIAssistantConversationReadyExecutor.cpp b/Engine/Plugins/Experimental/AIAssistant/Source/AIAssistant/Private/AIAssistantConversationReadyExecutor.cpp
--- a/Engine/Plugins/Experimental/AIAssistant/Source/AIAssistant/Private/AIAssistantConversationReadyExecutor.cpp
+++ b/Engine/Plugins/Experimental/AIAssistant/Source/AIAssistant/Private/AIAssistantConversationReadyExecutor.cpp
@@ -13,9 +13,16 @@ namespace UE::AIAssistant
 		FGetExecuteWhenReadyStateFunc&& GetExecuteWhenReadyStateFunc) :
 			GetExecuteWhenReadyStateDelegate(MoveTemp(GetExecuteWhenReadyStateFunc))
 	{
+		Reset();
+	}
+
+	void FConversationReadyExecutor::Reset()
+	{
+		UE::TUniqueLock Lock(StateMutex);
+		bConfiguredAgentEnvironment = false;
+		bCreatingConversation = false;
 	}
 
-	// Notify the executor that the agent environment has been configured.
 	void FConversationReadyExecutor::NotifyAgentEnvironmentConfigured()
 	{
 		{
diff --git a/Engine/Plugins/Experimental/AIAssistant/Source/AIAssistant/Private/AIAssistantConversationReadyExecutor.h b/Engine/Plugins/Experimental/AIAssistant/Source/AIAssistant/Private/AIAssistantConversationReadyExecutor.h
--- a/Engine/Plugins/Experimental/AIAssistant/Source/AIAssistant/Private/AIAssistantConversationReadyExecutor.h
+++ b/Engine/Plugins/Experimental/AIAssistant/Source/AIAssistant/Private/AIAssistantConversationReadyExecutor.h
@@ -23,6 +23,9 @@ namespace UE::AIAssistant
 
 		virtual ~FConversationReadyExecutor() = default;
 
+		// Reset to the initial state.
+		void Reset();
+
 		// Notify the executor that the agent environment has been configured.
 		void NotifyAgentEnvironmentConfigured();
 
@@ -44,8 +47,8 @@ namespace UE::AIAssistant
 
 	private:
 		mutable UE::FMutex StateMutex;
-		bool bConfiguredAgentEnvironment = false;
-		bool bCreatingConversation = false;
+		bool bConfiguredAgentEnvironment;
+		bool bCreatingConversation;
 		FGetExecuteWhenReadyStateFunc GetExecuteWhenReadyStateDelegate;
 
 	};
diff --git a/Engine/Plugins/Experimental/AIAssistant/Source/AIAssistant/Private/Tests/AIAssistantConversationReadyExecutorTest.cpp b/Engine/Plugins/Experimental/AIAssistant/Source/AIAssistant/Private/Tests/AIAssistantConversationReadyExecutorTest.cpp
--- a/Engine/Plugins/Experimental/AIAssistant/Source/AIAssistant/Private/Tests/AIAssistantConversationReadyExecutorTest.cpp
+++ b/Engine/Plugins/Experimental/AIAssistant/Source/AIAssistant/Private/Tests/AIAssistantConversationReadyExecutorTest.cpp
@@ -56,6 +56,36 @@ bool FAIAssistantConversationReadyExecutorExecuteOnConstructionTest::RunTest(
 	return true;
 }
 
+IMPLEMENT_SIMPLE_AUTOMATION_TEST(
+	FAIAssistantConversationReadyExecutorExecuteAfterResetTest,
+	"AI.Assistant.ConversationReadyExecutor.ExecuteAfterReset",
+	AIAssistantTest::Flags);
+
+bool FAIAssistantConversationReadyExecutorExecuteAfterResetTest::RunTest(
+	const FString& UnusedParameters)
+{
+	FConversationReadyExecutor ConversationReadyExecutor;
+	ConversationReadyExecutor.NotifyAgentEnvironmentConfigured();
+	// Block execution.
+	(void)ConversationReadyExecutor.SetCreatingConversation(true);
+
+	int ExecutedCount = 0;
+	ConversationReadyExecutor.ExecuteWhenReady([&ExecutedCount]() -> void { ++ExecutedCount; });
+	// Should not execute anything.
+	ConversationReadyExecutor.Reset();
+	(void)TestEqual(TEXT("NothingExecuted"), ExecutedCount, 0);
+
+	(void)TestFalse(TEXT("WasReset"), ConversationReadyExecutor.SetCreatingConversation(true));
+	// Notification should execute the enqueued operation on ready notification.
+	ConversationReadyExecutor.NotifyAgentEnvironmentConfigured();
+	(void)TestEqual(TEXT("ExecutedOnce"), ExecutedCount, 1);
+
+	// Should execute anything else now we've notified that the conversation is ready.
+	ConversationReadyExecutor.ExecuteWhenReady([&ExecutedCount]() -> void { ++ExecutedCount; });
+	(void)TestEqual(TEXT("ExecutedTwice"), ExecutedCount, 2);
+	return true;
+}
+
 IMPLEMENT_SIMPLE_AUTOMATION_TEST(
 	FAIAssistantConversationReadyExecutorExecuteAfterConversationCreatedTest,
 	"AI.Assistant.ConversationReadyExecutor.ExecuteAfterConversationCreated",
