[AI Assistant] Move temporary directory utility class out of config test.

Moved the temporary directory utility out of the config test so that it can be reused by
other tests.

diff --git a/Engine/Plugins/Experimental/AIAssistant/Source/AIAssistant/Private/Tests/AIAssistantConfigTest.cpp b/Engine/Plugins/Experimental/AIAssistant/Source/AIAssistant/Private/Tests/AIAssistantConfigTest.cpp
--- a/Engine/Plugins/Experimental/AIAssistant/Source/AIAssistant/Private/Tests/AIAssistantConfigTest.cpp
+++ b/Engine/Plugins/Experimental/AIAssistant/Source/AIAssistant/Private/Tests/AIAssistantConfigTest.cpp
@@ -4,49 +4,15 @@
 #include "Misc/AssertionMacros.h"
 #include "Misc/AutomationTest.h"
 #include "Misc/FileHelper.h"
-#include "Misc/Guid.h"
 #include "Misc/Paths.h"
-#include "HAL/FileManager.h"
 
 #include "AIAssistantConfig.h"
+#include "Tests/AIAssistantTemporaryDirectory.h"
 #include "Tests/AIAssistantTestFlags.h"
 
 #if WITH_DEV_AUTOMATION_TESTS
 
-// Create a temporary directory for the lifetime of this class.
-class FTemporaryDirectory
-{
-public:
-	
-	// Create a temporary directory.
-	FTemporaryDirectory() :
-		// NOTE: This uses a similar path to ContentCommandlets as there doesn't seem to be functionality
-		// to get the platform's temporary directory.
-		Directory(FPaths::Combine(FPaths::ProjectSavedDir(), TEXT("Temp"), FGuid::NewGuid().ToString())),
-		FileManager(IFileManager::Get())
-	{
-		verify(FileManager.MakeDirectory(*Directory, true));
-	}
-
-	// Disable copy.
-	FTemporaryDirectory(const FTemporaryDirectory&) = delete;
-	FTemporaryDirectory& operator=(const FTemporaryDirectory&) = delete;
-
-	// Delete a temporary directory.
-	~FTemporaryDirectory()
-	{
-		// ProjectSavedDir()/Temp will still remain, but may be empty.
-		FileManager.DeleteDirectory(*Directory, /* Exists */ 0, /* Tree */ 1);
-	}
-
-	// Get the temporary directory name.
-	const FString& operator*() const { return Directory; }
-
-private:
-	
-	FString Directory;
-	IFileManager& FileManager;
-};
+using namespace UE::AIAssistant;
 
 IMPLEMENT_SIMPLE_AUTOMATION_TEST(
 	FAIAssistantConfigTestHasDefaultSearchDirectories,
diff --git a/Engine/Plugins/Experimental/AIAssistant/Source/AIAssistant/Private/Tests/AIAssistantTemporaryDirectory.cpp b/Engine/Plugins/Experimental/AIAssistant/Source/AIAssistant/Private/Tests/AIAssistantTemporaryDirectory.cpp
new file mode 100644
--- /dev/null
+++ b/Engine/Plugins/Experimental/AIAssistant/Source/AIAssistant/Private/Tests/AIAssistantTemporaryDirectory.cpp
@@ -0,0 +1,31 @@
+// Copyright Epic Games, Inc. All Rights Reserved.
+
+#include "Tests/AIAssistantTemporaryDirectory.h"
+
+#include "Misc/AssertionMacros.h"
+#include "Containers/UnrealString.h"
+#include "Misc/Guid.h"
+#include "Misc/Paths.h"
+#include "HAL/FileManager.h"
+
+namespace UE::AIAssistant
+{
+	// Create a temporary directory.
+	FTemporaryDirectory::FTemporaryDirectory() :
+		// NOTE: This uses a similar path to ContentCommandlets as there doesn't seem to be functionality
+		// to get the platform's temporary directory.
+		Directory(
+			FPaths::Combine(
+				FPaths::ProjectSavedDir(), TEXT("Temp"), FGuid::NewGuid().ToString())),
+		FileManager(IFileManager::Get())
+	{
+		verify(FileManager.MakeDirectory(*Directory, true));
+	}
+
+	// Delete a temporary directory.
+	FTemporaryDirectory::~FTemporaryDirectory()
+	{
+		// ProjectSavedDir()/Temp will still remain, but may be empty.
+		FileManager.DeleteDirectory(*Directory, /* Exists */ false, /* Tree */ true);
+	}
+}
\ No newline at end of file
diff --git a/Engine/Plugins/Experimental/AIAssistant/Source/AIAssistant/Private/Tests/AIAssistantTemporaryDirectory.h b/Engine/Plugins/Experimental/AIAssistant/Source/AIAssistant/Private/Tests/AIAssistantTemporaryDirectory.h
new file mode 100644
--- /dev/null
+++ b/Engine/Plugins/Experimental/AIAssistant/Source/AIAssistant/Private/Tests/AIAssistantTemporaryDirectory.h
@@ -0,0 +1,32 @@
+// Copyright Epic Games, Inc. All Rights Reserved.
+
+#pragma once
+
+#include "Containers/UnrealString.h"
+
+class IFileManager;
+
+namespace UE::AIAssistant
+{
+	// Create a temporary directory for the lifetime of this class.
+	class FTemporaryDirectory
+	{
+	public:
+		// Create a temporary directory.
+		FTemporaryDirectory();
+
+		// Disable copy.
+		FTemporaryDirectory(const FTemporaryDirectory&) = delete;
+		FTemporaryDirectory& operator=(const FTemporaryDirectory&) = delete;
+
+		// Delete a temporary directory.
+		~FTemporaryDirectory();
+
+		// Get the temporary directory name.
+		const FString& operator*() const { return Directory; }
+
+	private:
+		FString Directory;
+		IFileManager& FileManager;
+	};
+}
\ No newline at end of file
