[AI Assistant] Moved web browser configuration into its own class and added tests.

This will be used by an assistant navigation policy class in a follow up CL making it
easier to test how to handle different page transitions without having to instantiate
a web browser.

diff --git a/Engine/Plugins/Experimental/AIAssistant/Source/AIAssistant/Private/AIAssistantWebBrowserConfig.cpp b/Engine/Plugins/Experimental/AIAssistant/Source/AIAssistant/Private/AIAssistantWebBrowserConfig.cpp
new file mode 100644
--- /dev/null
+++ b/Engine/Plugins/Experimental/AIAssistant/Source/AIAssistant/Private/AIAssistantWebBrowserConfig.cpp
@@ -0,0 +1,52 @@
+// Copyright Epic Games, Inc. All Rights Reserved.
+
+#include "AIAssistantWebBrowserConfig.h"
+
+#include "Containers/UnrealString.h"
+#include "Templates/UnrealTemplate.h"
+
+
+namespace UE::AIAssistant
+{
+	FWebBrowserConfig::FWebBrowserConfig(FAIAssistantConfig&& Config) :
+		FAIAssistantConfig(MoveTemp(Config))
+	{
+		// Cache regex patterns to avoid reallocation / compilation.
+		AllowedUrlRegexPatterns = GetAllowedUrlRegexPatterns();
+	}
+
+	EOpenBrowserMode FWebBrowserConfig::OpenBrowserModeFromNavigation(
+		const FString& Url, bool bUrlIsFromRedirect,
+		const FString& LastOpenedNonRedirectUrl) const
+	{
+		EOpenBrowserMode OpenBrowserMode = OpenBrowserModeFromUrl(Url);
+
+		// Allow redirects from allowed pages, to be loaded in the embedded browser.
+		//
+		// For example, the following is navigation flow that would load pages in the embedded
+		// browser:
+		// * Navigate to http://allowed.com - load in embedded browser
+		// * Redirect to http://not-explicitly-allowed.com  - load in embedded browser
+		// * Redirect to http://allowed.com/logged-in - load in embedded browser
+		// 
+		// The following is navigation flow that would prevent a page from being loaded in the
+		// embedded browser:
+		// * Navigate to http://not-allowed.com - load in embedded browser
+		// * Navigate to http://not-explicitly-allowed.com  - load in system browser
+		if (OpenBrowserMode == EOpenBrowserMode::System && bUrlIsFromRedirect)
+		{
+			OpenBrowserMode = OpenBrowserModeFromUrl(LastOpenedNonRedirectUrl);
+		}
+		return OpenBrowserMode;
+	}
+
+	EOpenBrowserMode FWebBrowserConfig::OpenBrowserModeFromUrl(
+		const FString& Url) const
+	{
+		for (const FRegexPattern& Pattern : AllowedUrlRegexPatterns)
+		{
+			if (FRegexMatcher(Pattern, Url).FindNext()) return EOpenBrowserMode::Embedded;
+		}
+		return EOpenBrowserMode::System;
+	}
+}
\ No newline at end of file
diff --git a/Engine/Plugins/Experimental/AIAssistant/Source/AIAssistant/Private/AIAssistantWebBrowserConfig.h b/Engine/Plugins/Experimental/AIAssistant/Source/AIAssistant/Private/AIAssistantWebBrowserConfig.h
new file mode 100644
--- /dev/null
+++ b/Engine/Plugins/Experimental/AIAssistant/Source/AIAssistant/Private/AIAssistantWebBrowserConfig.h
@@ -0,0 +1,38 @@
+// Copyright Epic Games, Inc. All Rights Reserved.
+
+#pragma once
+
+#include "Containers/Array.h"
+#include "Containers/UnrealString.h"
+#include "Internationalization/Regex.h"
+
+#include "AIAssistantConfig.h"
+
+namespace UE::AIAssistant
+{
+	// Where to load a URL.
+	enum class EOpenBrowserMode
+	{
+		Embedded,  // Open in the Unreal embedded browser.
+		System,    // Open in the system web browser.
+	};
+
+	// Loads configuration and determines whether URLs can be loaded in the embedded browser.
+	class FWebBrowserConfig : public FAIAssistantConfig
+	{
+	public:
+		explicit FWebBrowserConfig(FAIAssistantConfig&& Config);
+
+		// Determine the browser that should be used to open the specified URL as the result
+		// of a navigation event.
+		EOpenBrowserMode OpenBrowserModeFromNavigation(
+			const FString& Url, bool bUrlIsFromRedirect,
+			const FString& LastOpenedNonRedirectUrl) const;
+
+		// Determine the browser used to open a URL purely from the allow list.
+		EOpenBrowserMode OpenBrowserModeFromUrl(const FString& Url) const;
+
+	private:
+		TArray<FRegexPattern> AllowedUrlRegexPatterns;
+	};
+}
diff --git a/Engine/Plugins/Experimental/AIAssistant/Source/AIAssistant/Private/Tests/AIAssistantWebBrowserConfigTest.cpp b/Engine/Plugins/Experimental/AIAssistant/Source/AIAssistant/Private/Tests/AIAssistantWebBrowserConfigTest.cpp
new file mode 100644
--- /dev/null
+++ b/Engine/Plugins/Experimental/AIAssistant/Source/AIAssistant/Private/Tests/AIAssistantWebBrowserConfigTest.cpp
@@ -0,0 +1,127 @@
+// Copyright Epic Games, Inc. All Rights Reserved.
+
+#include "Containers/UnrealString.h"
+#include "Misc/AutomationTest.h"
+#include "Misc/Paths.h"
+#include "Templates/UnrealTemplate.h"
+
+#include "AIAssistantConfig.h"
+#include "AIAssistantWebBrowserConfig.h"
+#include "Tests/AIAssistantTestFlags.h"
+
+using namespace UE::AIAssistant;
+
+#if WITH_DEV_AUTOMATION_TESTS
+
+static const FString MainUrl = TEXT("https://fake.url");
+static const FString AllowedUrl = TEXT("https://another.fake.url");
+static const FString AllowedUrlRegex = TEXT(R"regex(https://another\.fake\.url)regex");
+static const FString BlockedUrl = TEXT("https://unvalidated.url");
+static const FString AnotherBlockedUrl = TEXT("https://also.unvalidated.url");
+
+static FWebBrowserConfig CreateTestConfig()
+{
+	FAIAssistantConfig Config;
+	Config.MainUrl = MainUrl;
+	Config.AllowedUrlRegexes.Add(AllowedUrlRegex);
+	return FWebBrowserConfig(MoveTemp(Config));
+}
+
+IMPLEMENT_SIMPLE_AUTOMATION_TEST(
+	FAIAssistantWebBrowserConfigTestOpenBrowserModeFromUrlEmbedded,
+	"AI.Assistant.WebBrowserConfig.OpenBrowserModeFromUrlEmbedded",
+	AIAssistantTest::Flags);
+
+bool FAIAssistantWebBrowserConfigTestOpenBrowserModeFromUrlEmbedded::RunTest(
+	const FString& UnusedParameters)
+{
+	FWebBrowserConfig Config = CreateTestConfig();
+	(void)TestEqual(
+		TEXT("MainUrl"), Config.OpenBrowserModeFromUrl(MainUrl),
+		EOpenBrowserMode::Embedded);
+	(void)TestEqual(
+		TEXT("AllowedUrl"), Config.OpenBrowserModeFromUrl(AllowedUrl),
+		EOpenBrowserMode::Embedded);
+	return true;
+}
+
+IMPLEMENT_SIMPLE_AUTOMATION_TEST(
+	FAIAssistantWebBrowserConfigTestOpenBrowserModeFromUrlSystem,
+	"AI.Assistant.WebBrowserConfig.OpenBrowserModeFromUrlSystem",
+	AIAssistantTest::Flags);
+
+bool FAIAssistantWebBrowserConfigTestOpenBrowserModeFromUrlSystem::RunTest(
+	const FString& UnusedParameters)
+{
+	FWebBrowserConfig Config = CreateTestConfig();
+	(void)TestEqual(
+		TEXT("BlockedUrl"), Config.OpenBrowserModeFromUrl(BlockedUrl),
+		EOpenBrowserMode::System);
+	return true;
+}
+
+IMPLEMENT_SIMPLE_AUTOMATION_TEST(
+	FAIAssistantWebBrowserConfigTestOpenBrowserModeFromNavigationEmbeddedNoRedirect,
+	"AI.Assistant.WebBrowserConfig.OpenBrowserModeFromNavigationEmbeddedNoRedirect",
+	AIAssistantTest::Flags);
+
+bool FAIAssistantWebBrowserConfigTestOpenBrowserModeFromNavigationEmbeddedNoRedirect::RunTest(
+	const FString& UnusedParameters)
+{
+	FWebBrowserConfig Config = CreateTestConfig();
+	(void)TestEqual(
+		TEXT("AllowedUrl"),
+		Config.OpenBrowserModeFromNavigation(AllowedUrl, false, TEXT("")),
+		EOpenBrowserMode::Embedded);
+	return true;
+}
+
+IMPLEMENT_SIMPLE_AUTOMATION_TEST(
+	FAIAssistantWebBrowserConfigTestOpenBrowserModeFromNavigationSystemNoRedirect,
+	"AI.Assistant.WebBrowserConfig.OpenBrowserModeFromNavigationSystemNoRedirect",
+	AIAssistantTest::Flags);
+
+bool FAIAssistantWebBrowserConfigTestOpenBrowserModeFromNavigationSystemNoRedirect::RunTest(
+	const FString& UnusedParameters)
+{
+	FWebBrowserConfig Config = CreateTestConfig();
+	(void)TestEqual(
+		TEXT("BlockedUrl"),
+		Config.OpenBrowserModeFromNavigation(BlockedUrl, false, TEXT("")),
+		EOpenBrowserMode::System);
+	return true;
+}
+
+IMPLEMENT_SIMPLE_AUTOMATION_TEST(
+	FAIAssistantWebBrowserConfigTestOpenBrowserModeFromNavigationEmbeddedRedirect,
+	"AI.Assistant.WebBrowserConfig.OpenBrowserModeFromNavigationEmbeddedRedirect",
+	AIAssistantTest::Flags);
+
+bool FAIAssistantWebBrowserConfigTestOpenBrowserModeFromNavigationEmbeddedRedirect::RunTest(
+	const FString& UnusedParameters)
+{
+	FWebBrowserConfig Config = CreateTestConfig();
+	(void)TestEqual(
+		TEXT("BlockedUrl"),
+		Config.OpenBrowserModeFromNavigation(BlockedUrl, true, AllowedUrl),
+		EOpenBrowserMode::Embedded);
+	return true;
+}
+
+IMPLEMENT_SIMPLE_AUTOMATION_TEST(
+	FAIAssistantWebBrowserConfigTestOpenBrowserModeFromNavigationSystemRedirect,
+	"AI.Assistant.WebBrowserConfig.OpenBrowserModeFromNavigationSystemRedirect",
+	AIAssistantTest::Flags);
+
+bool FAIAssistantWebBrowserConfigTestOpenBrowserModeFromNavigationSystemRedirect::RunTest(
+	const FString& UnusedParameters)
+{
+	FWebBrowserConfig Config = CreateTestConfig();
+	(void)TestEqual(
+		TEXT("BlockedUrl"),
+		Config.OpenBrowserModeFromNavigation(BlockedUrl, true, AnotherBlockedUrl),
+		EOpenBrowserMode::System);
+	return true;
+}
+
+#endif
\ No newline at end of file
