[AI Assistant] Split web browser navigation policy from web browser.

Move web browser navigation policy into a separate class and added tests. The web
browser will be migrated to use this in a future CL.

diff --git a/Engine/Plugins/Experimental/AIAssistant/Source/AIAssistant/Private/AIAssistantWebBrowserNavigationPolicy.cpp b/Engine/Plugins/Experimental/AIAssistant/Source/AIAssistant/Private/AIAssistantWebBrowserNavigationPolicy.cpp
new file mode 100644
--- /dev/null
+++ b/Engine/Plugins/Experimental/AIAssistant/Source/AIAssistant/Private/AIAssistantWebBrowserNavigationPolicy.cpp
@@ -0,0 +1,40 @@
+// Copyright Epic Games, Inc. All Rights Reserved.
+
+#include "AIAssistantWebBrowserNavigationPolicy.h"
+
+#include "IWebBrowserWindow.h"
+#include "Templates/SharedPointer.h"
+
+#include "AIAssistantWebBrowserConfig.h"
+
+namespace UE::AIAssistant
+{
+	FWebBrowserNavigationPolicy::FWebBrowserNavigationPolicy(
+		TSharedRef<FWebBrowserConfig> WebBrowserConfig) :
+		Config(WebBrowserConfig)
+	{
+	}
+
+	EOpenBrowserMode FWebBrowserNavigationPolicy::Update(
+		const FString& Url, const FWebNavigationRequest& Request)
+	{
+		bOpenedUrlChanged = Url != LastOpenedUrl;
+		LastOpenedUrl = Url;
+
+		EOpenBrowserMode OpenBrowserMode =
+			Config->OpenBrowserModeFromNavigation(
+				Url, Request.bIsRedirect, LastOpenedNonRedirectUrl);
+
+		// If the page can be opened in the embedded browser keep track of the last non-redirect URL.
+		if (OpenBrowserMode == EOpenBrowserMode::Embedded && !Request.bIsRedirect)
+		{
+			LastOpenedNonRedirectUrl = Url;
+		}
+		return OpenBrowserMode;
+	}
+
+	bool FWebBrowserNavigationPolicy::IsCurrentUrlMainUrl() const
+	{
+		return !Config->MainUrl.IsEmpty() && LastOpenedUrl.StartsWith(Config->MainUrl);
+	}
+}
\ No newline at end of file
diff --git a/Engine/Plugins/Experimental/AIAssistant/Source/AIAssistant/Private/AIAssistantWebBrowserNavigationPolicy.h b/Engine/Plugins/Experimental/AIAssistant/Source/AIAssistant/Private/AIAssistantWebBrowserNavigationPolicy.h
new file mode 100644
--- /dev/null
+++ b/Engine/Plugins/Experimental/AIAssistant/Source/AIAssistant/Private/AIAssistantWebBrowserNavigationPolicy.h
@@ -0,0 +1,37 @@
+// Copyright Epic Games, Inc. All Rights Reserved.
+
+#pragma once
+
+#include "Containers/UnrealString.h"
+#include "Templates/SharedPointer.h"
+
+#include "AIAssistantWebBrowserConfig.h"
+
+struct FWebNavigationRequest;
+
+namespace UE::AIAssistant
+{
+	class FWebBrowserNavigationPolicy
+	{
+	public:
+		FWebBrowserNavigationPolicy(TSharedRef<FWebBrowserConfig> WebBrowserConfig);
+
+		// Determine the open browser mode and update the current navigation state.
+		EOpenBrowserMode Update(const FString& Url, const FWebNavigationRequest& Request);
+
+		// Has the opened URL changed since the last call to update.
+		bool HasOpenedUrlChanged() const { return bOpenedUrlChanged; }
+
+		// Whether the current URL is the main url in the config.
+		bool IsCurrentUrlMainUrl() const;
+
+	private:
+		TSharedRef<FWebBrowserConfig> Config;
+		// Previously opened URL.
+		FString LastOpenedUrl;
+		// Whether navigation changed.
+		bool bOpenedUrlChanged = false;
+		// Previously loaded URL that was not a redirect.
+		FString LastOpenedNonRedirectUrl;
+	};
+}
\ No newline at end of file
diff --git a/Engine/Plugins/Experimental/AIAssistant/Source/AIAssistant/Private/Tests/AIAssistantWebBrowserNavigationPolicyTest.cpp b/Engine/Plugins/Experimental/AIAssistant/Source/AIAssistant/Private/Tests/AIAssistantWebBrowserNavigationPolicyTest.cpp
new file mode 100644
--- /dev/null
+++ b/Engine/Plugins/Experimental/AIAssistant/Source/AIAssistant/Private/Tests/AIAssistantWebBrowserNavigationPolicyTest.cpp
@@ -0,0 +1,158 @@
+// Copyright Epic Games, Inc. All Rights Reserved.
+
+#include "IWebBrowserWindow.h"
+#include "Misc/AutomationTest.h"
+#include "Templates/SharedPointer.h"
+#include "Templates/UnrealTemplate.h"
+
+#include "AIAssistantConfig.h"
+#include "AIAssistantWebBrowserNavigationPolicy.h"
+#include "AIAssistantWebBrowserConfig.h"
+#include "Tests/AIAssistantTestFlags.h"
+
+using namespace UE::AIAssistant;
+
+#if WITH_DEV_AUTOMATION_TESTS
+
+static TSharedRef<FWebBrowserConfig> CreateWebBrowserConfig()
+{
+	FAIAssistantConfig Config;
+	Config.MainUrl = FAIAssistantConfig::DefaultMainUrl;
+	Config.AllowedUrlRegexes = FAIAssistantConfig::DefaultAllowedUrlRegexes;
+	return MakeShared<FWebBrowserConfig>(MoveTemp(Config));
+}
+
+static FWebNavigationRequest CreateWebNavigationRequest()
+{
+	FWebNavigationRequest Request;
+	Request.bIsRedirect = false;
+	Request.bIsMainFrame = true;
+	Request.bIsExplicitTransition = false;
+	Request.TransitionSource = EWebTransitionSource::Link;
+	Request.TransitionSourceQualifier = EWebTransitionSourceQualifier::Unknown;
+	return Request;
+}
+
+static FWebNavigationRequest CreateWebNavigationRequestForRedirect()
+{
+	FWebNavigationRequest Request = CreateWebNavigationRequest();
+	Request.bIsRedirect = true;
+	return Request;
+}
+
+IMPLEMENT_SIMPLE_AUTOMATION_TEST(
+	FAIAssistantWebBrowserNavigationPolicyTestIsMainUrlOnInit,
+	"AI.Assistant.WebBrowserNavigationPolicy.IsMainUrlOnInit",
+	AIAssistantTest::Flags);
+
+bool FAIAssistantWebBrowserNavigationPolicyTestIsMainUrlOnInit::RunTest(
+	const FString& UnusedParameters)
+{
+	FWebBrowserNavigationPolicy NavigationPolicy(CreateWebBrowserConfig());
+	(void)TestFalse(TEXT("CurrentIsNotMainUrl"), NavigationPolicy.IsCurrentUrlMainUrl());
+	return true;
+}
+
+IMPLEMENT_SIMPLE_AUTOMATION_TEST(
+	FAIAssistantWebBrowserNavigationPolicyTestHasOpenedUrlChangedOnInit,
+	"AI.Assistant.WebBrowserNavigationPolicy.HasOpenedUrlChangedOnInit",
+	AIAssistantTest::Flags);
+
+bool FAIAssistantWebBrowserNavigationPolicyTestHasOpenedUrlChangedOnInit::RunTest(
+	const FString& UnusedParameters)
+{
+	FWebBrowserNavigationPolicy NavigationPolicy(CreateWebBrowserConfig());
+	(void)TestFalse(TEXT("OpenedUrlHasNotChanged"), NavigationPolicy.HasOpenedUrlChanged());
+	return true;
+}
+
+IMPLEMENT_SIMPLE_AUTOMATION_TEST(
+	FAIAssistantWebBrowserNavigationPolicyTestUpdateToMainUrl,
+	"AI.Assistant.WebBrowserNavigationPolicy.UpdateToMainUrl",
+	AIAssistantTest::Flags);
+
+bool FAIAssistantWebBrowserNavigationPolicyTestUpdateToMainUrl::RunTest(
+	const FString& UnusedParameters)
+{
+	auto Config = CreateWebBrowserConfig();
+	FWebBrowserNavigationPolicy NavigationPolicy(Config);
+	(void)TestEqual(
+		TEXT("UseEmbeddedBrowser"),
+		NavigationPolicy.Update(Config->MainUrl, CreateWebNavigationRequest()),
+		EOpenBrowserMode::Embedded);
+	(void)TestTrue(TEXT("CurrentIsMainUrl"), NavigationPolicy.IsCurrentUrlMainUrl());
+	return true;
+}
+
+IMPLEMENT_SIMPLE_AUTOMATION_TEST(
+	FAIAssistantWebBrowserNavigationPolicyTestUpdateToEmbeddedRedirect,
+	"AI.Assistant.WebBrowserNavigationPolicy.UpdateToEmbeddedRedirect",
+	AIAssistantTest::Flags);
+
+bool FAIAssistantWebBrowserNavigationPolicyTestUpdateToEmbeddedRedirect::RunTest(
+	const FString& UnusedParameters)
+{
+	auto Config = CreateWebBrowserConfig();
+	FWebBrowserNavigationPolicy NavigationPolicy(Config);
+	(void)NavigationPolicy.Update(Config->MainUrl, CreateWebNavigationRequest());
+	(void)TestEqual(
+		TEXT("UseEmbeddedBrowser"),
+		NavigationPolicy.Update(TEXT("https://a.random.url"), CreateWebNavigationRequestForRedirect()),
+		EOpenBrowserMode::Embedded);
+	return true;
+}
+
+IMPLEMENT_SIMPLE_AUTOMATION_TEST(
+	FAIAssistantWebBrowserNavigationPolicyTestUpdateToSystem,
+	"AI.Assistant.WebBrowserNavigationPolicy.UpdateToSystem",
+	AIAssistantTest::Flags);
+
+bool FAIAssistantWebBrowserNavigationPolicyTestUpdateToSystem::RunTest(
+	const FString& UnusedParameters)
+{
+	auto Config = CreateWebBrowserConfig();
+	FWebBrowserNavigationPolicy NavigationPolicy(Config);
+	(void)NavigationPolicy.Update(Config->MainUrl, CreateWebNavigationRequest());
+	(void)TestEqual(
+		TEXT("UseSystemBrowser"),
+		NavigationPolicy.Update(
+			TEXT("https://another.random.url"),
+			CreateWebNavigationRequest()),
+		EOpenBrowserMode::System);
+	return true;
+}
+
+IMPLEMENT_SIMPLE_AUTOMATION_TEST(
+	FAIAssistantWebBrowserNavigationPolicyTestHasOpenedUrlChangedModified,
+	"AI.Assistant.WebBrowserNavigationPolicy.HasOpenedUrlChangedModified",
+	AIAssistantTest::Flags);
+
+bool FAIAssistantWebBrowserNavigationPolicyTestHasOpenedUrlChangedModified::RunTest(
+	const FString& UnusedParameters)
+{
+	auto Config = CreateWebBrowserConfig();
+	FWebBrowserNavigationPolicy NavigationPolicy(Config);
+	(void)NavigationPolicy.Update(Config->MainUrl, CreateWebNavigationRequest());
+	(void)TestTrue(TEXT("HasChanged"), NavigationPolicy.HasOpenedUrlChanged());
+	(void)NavigationPolicy.Update(TEXT("https://a.random.url"), CreateWebNavigationRequest());
+	(void)TestTrue(TEXT("HasChanged"), NavigationPolicy.HasOpenedUrlChanged());
+	return true;
+}
+
+IMPLEMENT_SIMPLE_AUTOMATION_TEST(
+	FAIAssistantWebBrowserNavigationPolicyTestHasOpenedUrlChangedUnmodified,
+	"AI.Assistant.WebBrowserNavigationPolicy.HasOpenedUrlChangedUnmodified",
+	AIAssistantTest::Flags);
+
+bool FAIAssistantWebBrowserNavigationPolicyTestHasOpenedUrlChangedUnmodified::RunTest(
+	const FString& UnusedParameters)
+{
+	auto Config = CreateWebBrowserConfig();
+	FWebBrowserNavigationPolicy NavigationPolicy(Config);
+	(void)NavigationPolicy.Update(Config->MainUrl, CreateWebNavigationRequest());
+	(void)NavigationPolicy.Update(Config->MainUrl, CreateWebNavigationRequest());
+	(void)TestFalse(TEXT("HasNotChanged"), NavigationPolicy.HasOpenedUrlChanged());
+	return true;
+}
+
+#endif
\ No newline at end of file
