[AI Assistant] Fix conversation ready reset test.

The test was not correctly configuring the class to trigger execution of enqueued operations.

diff --git a/Engine/Plugins/Experimental/AIAssistant/Source/AIAssistant/Private/Tests/AIAssistantConversationReadyExecutorTest.cpp b/Engine/Plugins/Experimental/AIAssistant/Source/AIAssistant/Private/Tests/AIAssistantConversationReadyExecutorTest.cpp
--- a/Engine/Plugins/Experimental/AIAssistant/Source/AIAssistant/Private/Tests/AIAssistantConversationReadyExecutorTest.cpp
+++ b/Engine/Plugins/Experimental/AIAssistant/Source/AIAssistant/Private/Tests/AIAssistantConversationReadyExecutorTest.cpp
@@ -75,8 +75,11 @@ bool FAIAssistantConversationReadyExecutorExecuteAfterResetTest::RunTest(
 	ConversationReadyExecutor.Reset();
 	(void)TestEqual(TEXT("NothingExecuted"), ExecutedCount, 0);
 
+	// Creating a new conversation flushes enqueued executions.
 	(void)TestFalse(TEXT("WasReset"), ConversationReadyExecutor.SetCreatingConversation(true));
-	// Notification should execute the enqueued operation on ready notification.
+	ConversationReadyExecutor.ExecuteWhenReady([&ExecutedCount]() -> void { ++ExecutedCount; });
+	// Notification should execute the enqueued operation.
+	(void)ConversationReadyExecutor.SetCreatingConversation(false);
 	ConversationReadyExecutor.NotifyAgentEnvironmentConfigured();
 	(void)TestEqual(TEXT("ExecutedOnce"), ExecutedCount, 1);
 
