[AI Assistant] Move fake web API out of the web API test.

Move the fake web API out of the web API test so that it can be used by other tests.
Also, generalize finding called functions and completing them rather than just handling
the first called function.

diff --git a/Engine/Plugins/Experimental/AIAssistant/Source/AIAssistant/Private/Tests/AIAssistantFakeWebApi.cpp b/Engine/Plugins/Experimental/AIAssistant/Source/AIAssistant/Private/Tests/AIAssistantFakeWebApi.cpp
new file mode 100644
--- /dev/null
+++ b/Engine/Plugins/Experimental/AIAssistant/Source/AIAssistant/Private/Tests/AIAssistantFakeWebApi.cpp
@@ -0,0 +1,83 @@
+// Copyright Epic Games, Inc. All Rights Reserved.
+
+#include "AIAssistantFakeWebApi.h"
+
+#if WITH_DEV_AUTOMATION_TESTS
+
+namespace UE::AIAssistant
+{
+	TUniquePtr<FFakeWebApi::FExecutedAsyncFunction> FFakeWebApi::TestExpectAsyncFunctionCall(
+		FAutomationTestBase& TestCase,
+		const TCHAR* FunctionName,
+		const TCHAR* Arguments) const
+	{
+		TUniquePtr<FExecutedAsyncFunction> ReturnFunction;
+		auto FoundFunctions = FindExecutedAsyncFunctions(FunctionName);
+		if (TestCase.TestEqual(FunctionName, 1, FoundFunctions.Num()))
+		{
+			const auto& FoundFunction = FoundFunctions[0];
+			bool bCalled;
+			bCalled = TestCase.TestEqual(
+				FString::Printf(TEXT("%s_Called"), FunctionName),
+				FoundFunction.FunctionName, FunctionName);
+			bCalled = TestCase.TestEqual(
+				FString::Printf(TEXT("%s_Args"), FunctionName),
+				FoundFunction.Arguments, Arguments) && bCalled;
+			ReturnFunction = MakeUnique<FExecutedAsyncFunction>(FoundFunction);
+		}
+		return ReturnFunction;
+	}
+
+	bool FFakeWebApi::TestExpectAsyncFunctionCallAndComplete(
+		FAutomationTestBase& TestCase,
+		const TCHAR* FunctionName,
+		const TCHAR* Arguments,
+		const TCHAR* ResultJson,
+		bool bResultJsonIsError)
+	{
+		auto FoundFunction = TestExpectAsyncFunctionCall(TestCase, FunctionName, Arguments);
+		if (!FoundFunction.IsValid()) return false;
+
+		CompleteAsyncFunction(*FoundFunction, ResultJson, bResultJsonIsError);
+		return true;
+	}
+
+	TArray<FFakeWebApi::FExecutedAsyncFunction> FFakeWebApi::FindExecutedAsyncFunctions(
+		const FString& FunctionName) const
+	{
+		TArray<FExecutedAsyncFunction> FoundFunctions;
+		for (const auto& Function : ExecutedAsyncFunctions)
+		{
+			if (Function.FunctionName == FunctionName) FoundFunctions.Add(Function);
+		}
+		return FoundFunctions;
+	}
+
+	void FFakeWebApi::RemoveAsyncFunctions(const FString& FunctionName)
+	{
+		ExecutedAsyncFunctions.RemoveAll(
+			[&FunctionName](FExecutedAsyncFunction& Function)
+			{
+				return Function.FunctionName == FunctionName;
+			});
+	}
+
+	void FFakeWebApi::CompleteAsyncFunction(
+		const FExecutedAsyncFunction& ExecutedAsyncFunction,
+		const FString& ResultJson, bool bResultJsonIsError)
+	{
+		FWebJavaScriptResultDelegateAccessor::CallHandleResult(
+			*WebJavaScriptResultDelegate, ExecutedAsyncFunction.HandlerId,
+			ResultJson, bResultJsonIsError);
+	}
+
+	void FFakeWebApi::ExecuteAsyncFunction(
+		const TCHAR* FunctionName, const TCHAR* Arguments, const TCHAR* HandlerId)
+	{
+		FWebApi::ExecuteAsyncFunction(FunctionName, Arguments, HandlerId);
+		ExecutedAsyncFunctions.Emplace(
+			FExecutedAsyncFunction{ FunctionName, Arguments, HandlerId });
+	}
+}
+
+#endif
\ No newline at end of file
diff --git a/Engine/Plugins/Experimental/AIAssistant/Source/AIAssistant/Private/Tests/AIAssistantFakeWebApi.h b/Engine/Plugins/Experimental/AIAssistant/Source/AIAssistant/Private/Tests/AIAssistantFakeWebApi.h
new file mode 100644
--- /dev/null
+++ b/Engine/Plugins/Experimental/AIAssistant/Source/AIAssistant/Private/Tests/AIAssistantFakeWebApi.h
@@ -0,0 +1,125 @@
+// Copyright Epic Games, Inc. All Rights Reserved.
+
+#pragma once
+
+#include "Containers/Array.h"
+#include "Containers/UnrealString.h"
+#include "Misc/AssertionMacros.h"
+#include "Misc/AutomationTest.h"
+#include "Templates/UniquePtr.h"
+#include "Templates/ValueOrError.h"
+
+#include "AIAssistantWebApi.h"
+#include "Tests/AIAssistantWebJavaScriptResultDelegateAccessor.h"
+
+#if WITH_DEV_AUTOMATION_TESTS
+
+namespace UE::AIAssistant
+{
+	class FFakeWebApi : public FWebApi
+	{
+	public:
+		// Function executed by the FWebJavaScriptFunctionExecutor used to construct FFakeWebApi.
+		struct FExecutedAsyncFunction
+		{
+			// Name of the executed function.
+			FString FunctionName;
+			// Arguments passed to the function.
+			FString Arguments;
+			// Handler for the result of the function.
+			FString HandlerId;
+		};
+
+	public:
+		using FWebApi::FWebApi;
+		virtual ~FFakeWebApi() = default;
+
+		// Ensure a JavaScript function was called with the specified arguments.
+		TUniquePtr<FExecutedAsyncFunction> TestExpectAsyncFunctionCall(
+			FAutomationTestBase& TestCase,
+			const TCHAR* FunctionName,
+			const TCHAR* Arguments = TEXT("")) const;
+
+		// Ensure a JavaScript function was called with the specified arguments and complete it
+		// with ResultJson as a result of error if bResultJsonIsError is true.
+		bool TestExpectAsyncFunctionCallAndComplete(
+			FAutomationTestBase& TestCase,
+			const TCHAR* FunctionName,
+			const TCHAR* Arguments,
+			const TCHAR* ResultJson,
+			bool bResultJsonIsError);
+
+		// Ensure a JavaScript function was called with the specified arguments and complete it
+		// with ResultJson as a result of error if bResultJsonIsError is true. Also, validate the
+		// future receives the appropriate TValueOrError from ResultJson and bResultJsonIsError.
+		template<typename ResultFutureType>
+		bool TestExpectAsyncFunctionCallAndComplete(
+			FAutomationTestBase& TestCase,
+			const TCHAR* FunctionName,
+			const TCHAR* Arguments,
+			ResultFutureType& ResultFuture,
+			const TCHAR* ResultJson,
+			bool bResultJsonIsError)
+		{
+			if (!(TestCase.TestFalse(
+					FString::Printf(TEXT("%s_CompleteBeforeHandler"), FunctionName),
+					ResultFuture.IsReady()) &&
+				TestExpectAsyncFunctionCallAndComplete(
+					TestCase, FunctionName, Arguments, ResultJson, bResultJsonIsError) &&
+				TestCase.TestTrue(
+					FString::Printf(TEXT("%s_CompleteAfterHandler"), FunctionName),
+					ResultFuture.IsReady())))
+			{
+				return false;
+			}
+			const auto& Result = ResultFuture.Get();
+			bool bCompletedAsExpected;
+			bCompletedAsExpected = TestCase.TestEqual(
+				FString::Printf(TEXT("%s_ResultHasError"), FunctionName),
+				Result.HasError(), bResultJsonIsError);
+			bCompletedAsExpected = TestCase.TestEqual(
+				FString::Printf(TEXT("%s_HasExpectedResult"), FunctionName),
+				bResultJsonIsError ? Result.GetError() : GetValueOrEmptyString(Result),
+				ResultJson) && bCompletedAsExpected;
+			return bCompletedAsExpected;
+		}
+
+		// Find executed functions by name.
+		TArray<FExecutedAsyncFunction> FindExecutedAsyncFunctions(
+			const FString& FunctionName) const;
+
+		// Remove executed functions by name.
+		void RemoveAsyncFunctions(const FString& FunctionName);
+
+		// Complete an executed function.
+		void CompleteAsyncFunction(
+			const FExecutedAsyncFunction& ExecutedAsyncFunction,
+			const FString& ResultJson, bool bResultJsonIsError);
+
+	protected:
+		void ExecuteAsyncFunction(
+			const TCHAR* FunctionName, const TCHAR* Arguments, const TCHAR* HandlerId) override;
+
+		// If possible, get the value of a result as JSON.  If the value cannot be converted to
+		// JSON return an empty string.
+		template<typename ValueType>
+		FString GetValueOrEmptyString(const TValueOrError<ValueType, FString>& ValueOrError)
+		{
+			check(!ValueOrError.HasError());
+			return ValueOrError.GetValue().ToJson(false);
+		}
+
+		// Return an empty string for the value of a void TValueOrError.
+		template<>
+		FString GetValueOrEmptyString<void>(const TValueOrError<void, FString>& ValueOrError)
+		{
+			check(!ValueOrError.HasError());
+			return FString();
+		}
+
+	public:
+		TArray<FExecutedAsyncFunction> ExecutedAsyncFunctions;
+	};
+}
+
+#endif
\ No newline at end of file
diff --git a/Engine/Plugins/Experimental/AIAssistant/Source/AIAssistant/Private/Tests/AIAssistantWebApiTest.cpp b/Engine/Plugins/Experimental/AIAssistant/Source/AIAssistant/Private/Tests/AIAssistantWebApiTest.cpp
--- a/Engine/Plugins/Experimental/AIAssistant/Source/AIAssistant/Private/Tests/AIAssistantWebApiTest.cpp
+++ b/Engine/Plugins/Experimental/AIAssistant/Source/AIAssistant/Private/Tests/AIAssistantWebApiTest.cpp
@@ -6,6 +6,7 @@
 #include "Misc/AutomationTest.h"
 #include "Templates/Tuple.h"
 
+#include "AIAssistantFakeWebApi.h"
 #include "AIAssistantFakeWebJavaScriptExecutor.h"
 #include "AIAssistantFakeWebJavaScriptDelegateBinder.h"
 #include "AIAssistantTestFlags.h"
@@ -16,124 +17,20 @@
 
 using namespace UE::AIAssistant;
 
-class FTestWebApi : public FWebApi
+namespace UE::AIAssistant
 {
-public:
-	struct FExecutedAsyncFunction {
-		FString FunctionName;
-		FString Arguments;
-		FString HandlerId;
-	};
-
-public:
-	using FWebApi::FWebApi;
-	virtual ~FTestWebApi() = default;
-
-	// Ensure a JavaScript function was called with the specified arguments.
-	bool TestExpectAsyncFunctionCall(
-		FAutomationTestBase& TestCase,
-		const TCHAR* FunctionName,
-		const TCHAR* Arguments = TEXT("")) const
+	struct FFakeWebApiContainer
 	{
-		if (!TestCase.TestEqual(FunctionName, 1, ExecutedAsyncFunctions.Num()))
-		{
-			return false;
-		}
+		FFakeWebApiContainer() : WebApi(WebJavaScriptExecutor, WebJavaScriptDelegateBinder) {}
 
-		const auto& ExecutedAsyncFunction = ExecutedAsyncFunctions[0];
-		bool bCalled;
-		bCalled = TestCase.TestEqual(
-			FString::Printf(TEXT("%s_Called"), FunctionName),
-			ExecutedAsyncFunction.FunctionName, FunctionName);
-		bCalled = TestCase.TestEqual(
-			FString::Printf(TEXT("%s_Args"), FunctionName),
-			ExecutedAsyncFunction.Arguments, Arguments) && bCalled;
-		return bCalled;
-	}
-
-	// If possible, get the value of a result as JSON.  If the value cannot be converted to JSON
-	// return an empty string.
-	template<typename ValueType>
-	FString GetValueOrEmptyString(const TValueOrError<ValueType, FString>& ValueOrError)
-	{
-		return ValueOrError.GetValue().ToJson(false);
-	}
+		FFakeWebJavaScriptExecutor WebJavaScriptExecutor;
+		FFakeWebJavaScriptDelegateBinder WebJavaScriptDelegateBinder;
+		FFakeWebApi WebApi;
 
-	template<>
-	FString GetValueOrEmptyString<void>(const TValueOrError<void, FString>& UnusedValueOrError)
-	{
-		return FString();
-	}
-
-	// Ensure a JavaScript function was called with the specified arguments and complete it
-	// with ResultJson as a result of error if bResultJsonIsError is true.
-	template<typename ResultFutureType>
-	bool TestExpectAsyncFunctionCallAndComplete(
-		FAutomationTestBase& TestCase,
-		const TCHAR* FunctionName,
-		const TCHAR* Arguments,
-		ResultFutureType& ResultFuture,
-		const TCHAR* ResultJson,
-		bool bResultJsonIsError)
-	{
-		if (!TestExpectAsyncFunctionCall(TestCase, FunctionName, Arguments) || 
-			!TestCase.TestFalse(
-				FString::Printf(TEXT("%s_CompleteBeforeHandler"), FunctionName),
-				ResultFuture.IsReady()))
-		{
-			return false;
-		}
-		
-		FWebJavaScriptResultDelegateAccessor::CallHandleResult(
-			*WebJavaScriptResultDelegate,
-			ExecutedAsyncFunctions[0].HandlerId,
-			ResultJson, bResultJsonIsError);
-
-		if (!TestCase.TestTrue(
-				FString::Printf(TEXT("%s_CompleteAfterHandler"), FunctionName),
-				ResultFuture.IsReady()))
-		{
-			return false;
-		}
-		const auto& Result = ResultFuture.Get();
-		bool bCompletedAsExpected;
-		bCompletedAsExpected = TestCase.TestEqual(
-			FString::Printf(TEXT("%s_ResultHasError"), FunctionName),
-			Result.HasError(), bResultJsonIsError);
-		bCompletedAsExpected = TestCase.TestEqual(
-			FString::Printf(TEXT("%s_HasExpectedResult"), FunctionName),
-			bResultJsonIsError ? Result.GetError() : GetValueOrEmptyString(Result),
-			ResultJson) && bCompletedAsExpected;
-		return bCompletedAsExpected;
-	}
-
-protected:
-	void ExecuteAsyncFunction(
-		const TCHAR* FunctionName, const TCHAR* Arguments, const TCHAR* HandlerId) override
-	{
-		FWebApi::ExecuteAsyncFunction(FunctionName, Arguments, HandlerId);
-		ExecutedAsyncFunctions.Emplace(
-			FExecutedAsyncFunction{ FunctionName, Arguments, HandlerId });
-	}
-
-public:
-	TArray<FExecutedAsyncFunction> ExecutedAsyncFunctions;
-};
-
-struct FFakeWebApi
-{
-	FFakeWebApi() : WebApi(WebJavaScriptExecutor, WebJavaScriptDelegateBinder) {}
-
-	FFakeWebJavaScriptExecutor WebJavaScriptExecutor;
-	FFakeWebJavaScriptDelegateBinder WebJavaScriptDelegateBinder;
-	FTestWebApi WebApi;
-
-	FTestWebApi& operator*() { return WebApi; }
-	FTestWebApi* operator->() { return &WebApi; }
-};
+		FFakeWebApi& operator*() { return WebApi; }
+		FFakeWebApi* operator->() { return &WebApi; }
+	};
 
-namespace UE::AIAssistant
-{
 	class FWebApiAccessor
 	{
 	public:
@@ -165,7 +62,7 @@ IMPLEMENT_SIMPLE_AUTOMATION_TEST(
 
 bool FAIAssistantWebApiTestFormatFunctionCallNoArgs::RunTest(const FString& UnusedParameters)
 {
-	FFakeWebApi WebApi;
+	FFakeWebApiContainer WebApi;
 	return TestEqual(
 		TEXT("FormatFunctionCallNoArgs"),
 		FWebApiAccessor::FormatFunctionCall(*WebApi, TEXT("test"), TEXT("")),
@@ -191,7 +88,7 @@ IMPLEMENT_SIMPLE_AUTOMATION_TEST(
 
 bool FAIAssistantWebApiTestFormatFunctionCallWithArgs::RunTest(const FString& UnusedParameters)
 {
-	FFakeWebApi WebApi;
+	FFakeWebApiContainer WebApi;
 	return TestEqual(
 		TEXT("FormatFunctionCallWithArgs"),
 		FWebApiAccessor::FormatFunctionCall(*WebApi, TEXT("test"), TEXT("{foo: 'bar'}")),
@@ -218,7 +115,7 @@ IMPLEMENT_SIMPLE_AUTOMATION_TEST(
 bool FAIAssistantWebApiTestFormatResultAndErrorHandlers::RunTest(
 	const FString& UnusedParameters)
 {
-	FFakeWebApi WebApi;
+	FFakeWebApiContainer WebApi;
 	const FString FakeHandlerId = TEXT("foobar");
 	const auto Handlers = FWebApiAccessor::FormatResultAndErrorHandlers(*WebApi, FakeHandlerId);
 	const auto& JavaScriptResultDelegate = FWebApiAccessor::GetJavaScriptResultDelegate(*WebApi);
@@ -245,7 +142,7 @@ IMPLEMENT_SIMPLE_AUTOMATION_TEST(
 bool FAIAssistantWebApiTestFormatFunctionCallWithResultHandler::RunTest(
 	const FString& UnusedParameters)
 {
-	FFakeWebApi WebApi;
+	FFakeWebApiContainer WebApi;
 	const FString FakeHandlerId = TEXT("foobar");
 	const auto Handlers = FWebApiAccessor::FormatResultAndErrorHandlers(*WebApi, FakeHandlerId);
 	return TestEqual(
@@ -278,7 +175,7 @@ IMPLEMENT_SIMPLE_AUTOMATION_TEST(
 
 bool FAIAssistantWebApiTestCreateConversation::RunTest(const FString& UnusedParameters)
 {
-	FFakeWebApi WebApi;
+	FFakeWebApiContainer WebApi;
 	auto Result = WebApi->CreateConversation();
 	return WebApi->TestExpectAsyncFunctionCallAndComplete(
 		*this, TEXT("createConversation"), TEXT(""), Result, TEXT(""), false);
@@ -291,7 +188,7 @@ IMPLEMENT_SIMPLE_AUTOMATION_TEST(
 
 bool FAIAssistantWebApiTestAddMessageToConversation::RunTest(const FString& UnusedParameters)
 {
-	FFakeWebApi WebApi;
+	FFakeWebApiContainer WebApi;
 	FAddMessageToConversationOptions Options;
 	Options.ConversationId.Emplace().Id = TEXT("convo");
 	Options.Message.MessageRole = EMessageRole::User;
@@ -301,7 +198,7 @@ bool FAIAssistantWebApiTestAddMessageToConversation::RunTest(const FString& Unus
 	MessageContent.Content.Get<FTextMessageContent>().Text = TEXT("Hello");
 	WebApi->AddMessageToConversation(Options);
 	return WebApi->TestExpectAsyncFunctionCall(
-		*this, TEXT("addMessageToConversation"), *Options.ToJson(false));
+		*this, TEXT("addMessageToConversation"), *Options.ToJson(false)).IsValid();
 }
 
 IMPLEMENT_SIMPLE_AUTOMATION_TEST(
@@ -311,7 +208,7 @@ IMPLEMENT_SIMPLE_AUTOMATION_TEST(
 
 bool FAIAssistantWebApiTestAddAgentEnvironment::RunTest(const FString& UnusedParameters)
 {
-	FFakeWebApi WebApi;
+	FFakeWebApiContainer WebApi;
 	FAgentEnvironment AgentEnvironment;
 	AgentEnvironment.Descriptor.EnvironmentName = TEXT("UE");
 	AgentEnvironment.Descriptor.EnvironmentVersion = TEXT("5.7.0");
@@ -333,7 +230,7 @@ IMPLEMENT_SIMPLE_AUTOMATION_TEST(
 
 bool FAIAssistantWebApiTestAddAgentEnvironmentFailed::RunTest(const FString& UnusedParameters)
 {
-	FFakeWebApi WebApi;
+	FFakeWebApiContainer WebApi;
 	FAgentEnvironment AgentEnvironment;
 	auto Result = WebApi->AddAgentEnvironment(AgentEnvironment);
 
@@ -349,12 +246,12 @@ IMPLEMENT_SIMPLE_AUTOMATION_TEST(
 
 bool FAIAssistantWebApiTestSetAgentEnvironment::RunTest(const FString& UnusedParameters)
 {
-	FFakeWebApi WebApi;
+	FFakeWebApiContainer WebApi;
 	FAgentEnvironmentId Id;
 	Id.Id = TEXT("fakeId");
 	WebApi->SetAgentEnvironment(Id);
 	return WebApi->TestExpectAsyncFunctionCall(
-		*this, TEXT("setAgentEnvironment"), *Id.ToJson(false));
+		*this, TEXT("setAgentEnvironment"), *Id.ToJson(false)).IsValid();
 }
 
 IMPLEMENT_SIMPLE_AUTOMATION_TEST(
@@ -364,11 +261,12 @@ IMPLEMENT_SIMPLE_AUTOMATION_TEST(
 
 bool FAIAssistantWebApiTestUpdateGlobalLocale::RunTest(const FString& UnusedParameters)
 {
-	FFakeWebApi WebApi;
+	FFakeWebApiContainer WebApi;
 	FString LocaleFromSetting = TEXT("fr");
 	WebApi->UpdateGlobalLocale(LocaleFromSetting);
 	return WebApi->TestExpectAsyncFunctionCall(
-		*this, TEXT("updateGlobalLocale"), *FString::Printf(TEXT("\"%s\""), *LocaleFromSetting));
+		*this, TEXT("updateGlobalLocale"), *FString::Printf(TEXT("\"%s\""),
+		*LocaleFromSetting)).IsValid();
 }
 
 #endif  // WITH_DEV_AUTOMATION_TESTS
\ No newline at end of file
