[AI Assistant] Expose agent environment configuration state from conversation ready executor.

diff --git a/Engine/Plugins/Experimental/AIAssistant/Source/AIAssistant/Private/AIAssistantConversationReadyExecutor.h b/Engine/Plugins/Experimental/AIAssistant/Source/AIAssistant/Private/AIAssistantConversationReadyExecutor.h
--- a/Engine/Plugins/Experimental/AIAssistant/Source/AIAssistant/Private/AIAssistantConversationReadyExecutor.h
+++ b/Engine/Plugins/Experimental/AIAssistant/Source/AIAssistant/Private/AIAssistantConversationReadyExecutor.h
@@ -29,6 +29,9 @@ namespace UE::AIAssistant
 		// Notify the executor that the agent environment has been configured.
 		void NotifyAgentEnvironmentConfigured();
 
+		// Whether the agent environment has been configured.
+		bool IsAgentEnvironmentConfigured() const { return bConfiguredAgentEnvironment; }
+
 		// Set the creating conversation flag returning the previous value.
 		bool SetCreatingConversation(bool bNewCreatingConversation);
 
diff --git a/Engine/Plugins/Experimental/AIAssistant/Source/AIAssistant/Private/Tests/AIAssistantConversationReadyExecutorTest.cpp b/Engine/Plugins/Experimental/AIAssistant/Source/AIAssistant/Private/Tests/AIAssistantConversationReadyExecutorTest.cpp
--- a/Engine/Plugins/Experimental/AIAssistant/Source/AIAssistant/Private/Tests/AIAssistantConversationReadyExecutorTest.cpp
+++ b/Engine/Plugins/Experimental/AIAssistant/Source/AIAssistant/Private/Tests/AIAssistantConversationReadyExecutorTest.cpp
@@ -47,12 +47,17 @@ bool FAIAssistantConversationReadyExecutorExecuteOnConstructionTest::RunTest(
 	const FString& UnusedParameters)
 {
 	FConversationReadyExecutor ConversationReadyExecutor;
+	(void)TestFalse(
+		TEXT("NotConfigured"),
+		ConversationReadyExecutor.IsAgentEnvironmentConfigured());
+
 	bool bExecuted = false;
 	ConversationReadyExecutor.ExecuteWhenReady([&bExecuted]() -> void { bExecuted = true; });
 	(void)TestFalse(TEXT("Not configured on construction"), bExecuted);
 
 	ConversationReadyExecutor.NotifyAgentEnvironmentConfigured();
 	(void)TestTrue(TEXT("Ready when configured"), bExecuted);
+	(void)TestTrue(TEXT("Configured"), ConversationReadyExecutor.IsAgentEnvironmentConfigured());
 	return true;
 }
 
