[AI Assistant] Ensure each JavaScript result delegate has a unique name.

It will soon be possible to have multiple JavaScript result delegates handling
callbacks for function invocations. For example, when reloading a page with pending
JavaScript executions the current delegate is torn down asynchronously while another
is created. To ensure each result delegate is called with the appropriate context by the
browser ensure each intsance is assigned a unique name.

diff --git a/Engine/Plugins/Experimental/AIAssistant/Source/AIAssistant/Private/AIAssistantWebJavaScriptResultDelegate.cpp b/Engine/Plugins/Experimental/AIAssistant/Source/AIAssistant/Private/AIAssistantWebJavaScriptResultDelegate.cpp
--- a/Engine/Plugins/Experimental/AIAssistant/Source/AIAssistant/Private/AIAssistantWebJavaScriptResultDelegate.cpp
+++ b/Engine/Plugins/Experimental/AIAssistant/Source/AIAssistant/Private/AIAssistantWebJavaScriptResultDelegate.cpp
@@ -8,7 +8,7 @@
 // Name of this object when registered with the JavaScript binder.
 // NOTE: This is lower case as FCEFJSScripting::GetBindingName() will silently change the name
 // to lower case if FCEFJSScripting::bJSBindingToLoweringEnabled is true (the default).
-const FString UAIAssistantWebJavaScriptResultDelegate::Name(TEXT("aiassistantresultdelegate"));
+const FString UAIAssistantWebJavaScriptResultDelegate::BaseName(TEXT("aiassistantresultdelegate"));
 // Result set when a promise is canceled.
 const FString UAIAssistantWebJavaScriptResultDelegate::CanceledError(
 	TEXT(R"json("canceled")json"));
@@ -17,7 +17,7 @@ void UAIAssistantWebJavaScriptResultDelegate::Bind(
 	UE::AIAssistant::IWebJavaScriptDelegateBinder& WebJavaScriptDelegateBinder)
 {
 	ScopedWebJavaScriptDelegateBinder.Emplace(
-		WebJavaScriptDelegateBinder, Name, this, true /* bIsPermanent */);
+		WebJavaScriptDelegateBinder, GetName(), this, true /* bIsPermanent */);
 }
 
 void UAIAssistantWebJavaScriptResultDelegate::Unbind()
@@ -28,7 +28,7 @@ void UAIAssistantWebJavaScriptResultDelegate::Unbind()
 
 FString UAIAssistantWebJavaScriptResultDelegate::CreateHandlerId()
 {
-	return FGuid::NewGuid().ToString();
+	return FGuid::NewGuid().ToString(EGuidFormats::DigitsLower);
 }
 
 void UAIAssistantWebJavaScriptResultDelegate::RegisterResultHandlerWithId(
@@ -83,12 +83,23 @@ FString UAIAssistantWebJavaScriptResultDelegate::FormatJavaScriptHandler(const F
 {
 	return FString::Printf(
 		TEXT(R"js(window.ue.%s.%s("%s", JSON.stringify({0}), {1});)js"),
-		*Name,
+		*GetName(),
 		*GET_FUNCTION_NAME_CHECKED(
 			UAIAssistantWebJavaScriptResultDelegate, HandleResult).ToString().ToLower(),
 		*HandlerId);
 }
 
+const FString& UAIAssistantWebJavaScriptResultDelegate::GetName() const
+{
+	if (Name.IsEmpty())
+	{
+		Name = FString::Printf(
+			TEXT("%s_%s"), *BaseName,
+			*FGuid::NewGuid().ToString(EGuidFormats::DigitsLower));
+	}
+	return Name;
+}
+
 void UAIAssistantWebJavaScriptResultDelegate::BeginDestroy()
 {
 	Super::BeginDestroy();
diff --git a/Engine/Plugins/Experimental/AIAssistant/Source/AIAssistant/Private/AIAssistantWebJavaScriptResultDelegate.h b/Engine/Plugins/Experimental/AIAssistant/Source/AIAssistant/Private/AIAssistantWebJavaScriptResultDelegate.h
--- a/Engine/Plugins/Experimental/AIAssistant/Source/AIAssistant/Private/AIAssistantWebJavaScriptResultDelegate.h
+++ b/Engine/Plugins/Experimental/AIAssistant/Source/AIAssistant/Private/AIAssistantWebJavaScriptResultDelegate.h
@@ -80,6 +80,9 @@ public:
 	// ResultHandler and "{1}" which should be truthy if the result is an error, falsey otherwise.
 	FString FormatJavaScriptHandler(const FString& HandlerId) const;
 
+	// Get the name of the delegate.
+	const FString& GetName() const;
+
 protected:
 	// Handle the result of an execution.
 	// HandlerId is the unique ID of the execution, ResultJson is the result of the execution or
@@ -106,6 +109,9 @@ private:
 	void TryCompletePendingPromise(FResultHandlerContext&& ResultContext);
 
 private:
+	// Name of the delegate that is lazily initialized to avoid exposing this as a UObject
+	// property.
+	mutable FString Name;
 	UE::FMutex ResultHandlersByIdLock;  // Guards ResultHandlersById
 	TMap<FString, TSharedPtr<FResultHandler>> ResultHandlersById;
 	TOptional<UE::AIAssistant::FScopedWebJavaScriptDelegateBinder>
@@ -118,8 +124,8 @@ private:
 	static FString CreateHandlerId();
 
 public:
-	// Name of this object when registered with the JavaScript binder.
-	static const FString Name;
+	// Base name of this object when registered with the JavaScript binder.
+	static const FString BaseName;
 	// JSON result when a promise / future is canceled.
 	static const FString CanceledError;
 };
\ No newline at end of file
diff --git a/Engine/Plugins/Experimental/AIAssistant/Source/AIAssistant/Private/Tests/AIAssistantWebJavaScriptResultDelegateTest.cpp b/Engine/Plugins/Experimental/AIAssistant/Source/AIAssistant/Private/Tests/AIAssistantWebJavaScriptResultDelegateTest.cpp
--- a/Engine/Plugins/Experimental/AIAssistant/Source/AIAssistant/Private/Tests/AIAssistantWebJavaScriptResultDelegateTest.cpp
+++ b/Engine/Plugins/Experimental/AIAssistant/Source/AIAssistant/Private/Tests/AIAssistantWebJavaScriptResultDelegateTest.cpp
@@ -17,6 +17,33 @@
 
 using namespace UE::AIAssistant;
 
+IMPLEMENT_SIMPLE_AUTOMATION_TEST(
+	FAIAssistantWebJavaScriptResultDelegateTestGetName,
+	"AI.Assistant.WebJavaScriptResultDelegate.GetName",
+	AIAssistantTest::Flags);
+
+bool FAIAssistantWebJavaScriptResultDelegateTestGetName::RunTest(
+	const FString& UnusedParameters)
+{
+	TStrongObjectPtr<UAIAssistantWebJavaScriptResultDelegate> ResultDelegate(
+		NewObject<UAIAssistantWebJavaScriptResultDelegate>());
+	FString Name = ResultDelegate->GetName();
+	(void)TestEqual(
+		TEXT("NameStartsWithBaseName"),
+		Name.Left(UAIAssistantWebJavaScriptResultDelegate::BaseName.Len()),
+		UAIAssistantWebJavaScriptResultDelegate::BaseName);
+	(void)TestNotEqual(
+		TEXT("NameHasUniqueTail"),
+		Name.Right(
+			Name.Len() -
+			Name.Left(UAIAssistantWebJavaScriptResultDelegate::BaseName.Len()).Len()),
+		TEXT(""));
+
+	const FString& AnotherName = ResultDelegate->GetName();
+	(void)TestEqual(TEXT("NameDoesNotChange"), Name, AnotherName);
+	return true;
+}
+
 IMPLEMENT_SIMPLE_AUTOMATION_TEST(
 	FAIAssistantWebJavaScriptResultDelegateTestBindUnbind,
 	"AI.Assistant.WebJavaScriptResultDelegate.BindUnbind",
@@ -174,10 +201,12 @@ bool FAIAssistantWebJavaScriptResultDelegateTestFormatJavaScriptHandler::RunTest
 	FString HandlerId = TEXT("fake_handler_id");
 	(void)TestEqual(
 		TEXT("JavaScript"),
-		TEXT(
-			R"js(window.ue.aiassistantresultdelegate.handleresult)js"
-			R"js(("fake_handler_id", JSON.stringify({0}), {1});)js"),
-		ResultDelegate->FormatJavaScriptHandler(HandlerId));
+		ResultDelegate->FormatJavaScriptHandler(HandlerId),
+		FString::Printf(
+			TEXT(
+				R"js(window.ue.%s.handleresult)js"
+				R"js(("fake_handler_id", JSON.stringify({0}), {1});)js"),
+			*ResultDelegate->GetName()));
 	return true;
 }
 
