[AI Assistant] Fix a potential race when creating a JavaScript result handler.

diff --git a/Engine/Plugins/Experimental/AIAssistant/Source/AIAssistant/Private/AIAssistantWebJavaScriptResultDelegate.cpp b/Engine/Plugins/Experimental/AIAssistant/Source/AIAssistant/Private/AIAssistantWebJavaScriptResultDelegate.cpp
--- a/Engine/Plugins/Experimental/AIAssistant/Source/AIAssistant/Private/AIAssistantWebJavaScriptResultDelegate.cpp
+++ b/Engine/Plugins/Experimental/AIAssistant/Source/AIAssistant/Private/AIAssistantWebJavaScriptResultDelegate.cpp
@@ -91,6 +91,7 @@ FString UAIAssistantWebJavaScriptResultDelegate::FormatJavaScriptHandler(const F
 
 const FString& UAIAssistantWebJavaScriptResultDelegate::GetName() const
 {
+	UE::TUniqueLock Lock(NameLock);
 	if (Name.IsEmpty())
 	{
 		Name = FString::Printf(
diff --git a/Engine/Plugins/Experimental/AIAssistant/Source/AIAssistant/Private/AIAssistantWebJavaScriptResultDelegate.h b/Engine/Plugins/Experimental/AIAssistant/Source/AIAssistant/Private/AIAssistantWebJavaScriptResultDelegate.h
--- a/Engine/Plugins/Experimental/AIAssistant/Source/AIAssistant/Private/AIAssistantWebJavaScriptResultDelegate.h
+++ b/Engine/Plugins/Experimental/AIAssistant/Source/AIAssistant/Private/AIAssistantWebJavaScriptResultDelegate.h
@@ -109,6 +109,7 @@ private:
 	void TryCompletePendingPromise(FResultHandlerContext&& ResultContext);
 
 private:
+	mutable UE::FMutex NameLock; // Guards Name
 	// Name of the delegate that is lazily initialized to avoid exposing this as a UObject
 	// property.
 	mutable FString Name;
