AI Assistant: Fix embedded assistant pages being opened in the system browser.

This fixes embedded assistant pagees being opened in the system browser unless they're explicitly
requested via a pop-up ("_blank") target.

diff --git a/Engine/Plugins/Experimental/AIAssistant/Source/AIAssistant/Private/AIAssistantWebBrowserNavigationPolicy.cpp b/Engine/Plugins/Experimental/AIAssistant/Source/AIAssistant/Private/AIAssistantWebBrowserNavigationPolicy.cpp
--- a/Engine/Plugins/Experimental/AIAssistant/Source/AIAssistant/Private/AIAssistantWebBrowserNavigationPolicy.cpp
+++ b/Engine/Plugins/Experimental/AIAssistant/Source/AIAssistant/Private/AIAssistantWebBrowserNavigationPolicy.cpp
@@ -22,7 +22,9 @@ namespace UE::AIAssistant
 		LastOpenedUrl = Url;
 
 		EOpenBrowserMode OpenBrowserMode =
-			Config->OpenBrowserModeFromNavigation(
+			IsCurrentUrlMainUrl()
+			? EOpenBrowserMode::Embedded
+			: Config->OpenBrowserModeFromNavigation(
 				Url, Request.bIsRedirect, LastOpenedNonRedirectUrl);
 
 		// If the page can be opened in the embedded browser keep track of the last non-redirect URL.
diff --git a/Engine/Plugins/Experimental/AIAssistant/Source/AIAssistant/Private/Tests/AIAssistantWebBrowserNavigationPolicyTest.cpp b/Engine/Plugins/Experimental/AIAssistant/Source/AIAssistant/Private/Tests/AIAssistantWebBrowserNavigationPolicyTest.cpp
--- a/Engine/Plugins/Experimental/AIAssistant/Source/AIAssistant/Private/Tests/AIAssistantWebBrowserNavigationPolicyTest.cpp
+++ b/Engine/Plugins/Experimental/AIAssistant/Source/AIAssistant/Private/Tests/AIAssistantWebBrowserNavigationPolicyTest.cpp
@@ -84,6 +84,25 @@ bool FAIAssistantWebBrowserNavigationPolicyTestUpdateToMainUrl::RunTest(
 	return true;
 }
 
+IMPLEMENT_SIMPLE_AUTOMATION_TEST(
+	FAIAssistantWebBrowserNavigationPolicyTestUpdateToChildOfMainUrl,
+	"AI.Assistant.WebBrowserNavigationPolicy.UpdateToChildOfMainUrl",
+	AIAssistantTest::Flags);
+
+bool FAIAssistantWebBrowserNavigationPolicyTestUpdateToChildOfMainUrl::RunTest(
+	const FString& UnusedParameters)
+{
+	auto Config = CreateWebBrowserConfig();
+	FWebBrowserNavigationPolicy NavigationPolicy(Config);
+	(void)TestEqual(
+		TEXT("UseEmbeddedBrowser"),
+		NavigationPolicy.Update(
+			Config->MainUrl + TEXT("/another/page"), CreateWebNavigationRequest()),
+		EOpenBrowserMode::Embedded);
+	(void)TestTrue(TEXT("CurrentIsMainUrl"), NavigationPolicy.IsCurrentUrlMainUrl());
+	return true;
+}
+
 IMPLEMENT_SIMPLE_AUTOMATION_TEST(
 	FAIAssistantWebBrowserNavigationPolicyTestUpdateToEmbeddedRedirect,
 	"AI.Assistant.WebBrowserNavigationPolicy.UpdateToEmbeddedRedirect",
