[AI Assistant] Request Tooltips response to be in current editor language
- add engine version, language, locale to structured context
- make visible query localization responsive to runtime language changes
- remove localization from text that is not user visible

#rnx

diff --git a/Engine/Plugins/Experimental/AIAssistant/Source/AIAssistant/Private/AIAssistantSlateQuerier.cpp b/Engine/Plugins/Experimental/AIAssistant/Source/AIAssistant/Private/AIAssistantSlateQuerier.cpp
--- a/Engine/Plugins/Experimental/AIAssistant/Source/AIAssistant/Private/AIAssistantSlateQuerier.cpp
+++ b/Engine/Plugins/Experimental/AIAssistant/Source/AIAssistant/Private/AIAssistantSlateQuerier.cpp
@@ -14,6 +14,8 @@
 #include "EngineUtils.h"
 #include "IDetailsView.h"
 #include "Interfaces/IMainFrameModule.h"
+#include "Internationalization/Internationalization.h"
+#include "Internationalization/Culture.h"
 #include "Serialization/JsonSerializable.h"
 #include "Serialization/JsonSerializerMacros.h"
 #include "Subsystems/AssetEditorSubsystem.h"
@@ -39,6 +41,9 @@
 struct FAIAssistantSlateQueryStructuredContext : public FJsonSerializable
 {
 	FString TopWindowName;
+	FString UnrealEditorVersion;
+	FString UnrealEditorLanguage;
+	FString UnrealEditorLocale;
 	TOptional<FString> WindowName;
 	TOptional<FString> EditorName;
 	TOptional<FString> EditorMode;
@@ -55,6 +60,9 @@ struct FAIAssistantSlateQueryStructuredContext : public FJsonSerializable
 
 	BEGIN_JSON_SERIALIZER
 		JSON_SERIALIZE("toplevel_window_name", TopWindowName);
+		JSON_SERIALIZE("unreal_editor_version", UnrealEditorVersion);
+		JSON_SERIALIZE("unreal_editor_language", UnrealEditorLanguage);
+		JSON_SERIALIZE("unreal_editor_locale", UnrealEditorLocale);
 		JSON_SERIALIZE_OPTIONAL("window_name", WindowName);
 		JSON_SERIALIZE_OPTIONAL("editor_name", EditorName);
 		JSON_SERIALIZE_OPTIONAL("editor_mode", EditorMode);
@@ -96,6 +104,7 @@ struct FAIAssistantSlateQueryContext
 
 	FText GeneratedQuery = FText();
 	FText GeneratedQueryInstructions = FText();
+	FText LanguageInstructions = FText();
 
 	FAIAssistantSlateQueryStructuredContext StructuredContext;
 };
@@ -1131,7 +1140,6 @@ static FString CleanExtraWhiteSpaceFromString(const FString& InString)
 	return MoveTemp(CleanString);
 }
 
-
 //
 // UE::AIAssistant::SlateQuerier
 //
@@ -1199,23 +1207,25 @@ void UE::AIAssistant::SlateQuerier::QueryAIAssistantAboutSlateWidget(const FWidg
 			NameToQuery = (!FormattedItemName.IsEmpty()) ? FormattedItemName : (!SlateQueryContext.FormattedTabName.IsEmpty()) ? SlateQueryContext.FormattedTabName : SlateQueryContext.FormattedWindowName;
 		}
 
-		static const FText QueryFormatWidget = LOCTEXT("QueryFormat1", "I would like to know what {NameToQuery} does.");
-		static const FText QueryFormatGeneral = LOCTEXT("QueryFormat2", "I would like to know what \"{NameToQuery}\" means.");
-		static const FText QueryFormatObject = LOCTEXT("QueryFormat3", "I would like to know about {NameToQuery}.");
-		static const FText QueryFormatAssistant = LOCTEXT("QueryFormat4", "Tell me about yourself.");
-		static const FText QueryInstructionsWidgetFormat = LOCTEXT("QueryInstructionsFormat1",
+		// Localized text is non-static const so it can change when the editor language is changed.
+		const FText QueryFormatWidget = LOCTEXT("QueryFormat1", "I would like to know what {NameToQuery} does.");
+		const FText QueryFormatGeneral = LOCTEXT("QueryFormat2", "I would like to know what \"{NameToQuery}\" means.");
+		const FText QueryFormatObject = LOCTEXT("QueryFormat3", "I would like to know about {NameToQuery}.");
+		const FText QueryFormatAssistant = LOCTEXT("QueryFormat4", "Tell me about yourself.");
+		// Hidden instruction text sent to assistant can be non-localized and static const.
+		static const FText QueryInstructionsWidgetFormat = INVTEXT(
 			"Provide an easily readable, informative, accurate answer that describes what the widget does in the Unreal Editor UI. {QueryAdditionalInstructions} {StyleInstructions}");
-		static const FText QueryInstructionsObjectFormat = LOCTEXT("QueryInstructionsFormat2",
+		static const FText QueryInstructionsObjectFormat = INVTEXT(
 			"Provide an easily readable, informative, accurate answer that describes what I should know about it in the Unreal Editor UI. {QueryAdditionalInstructions} {StyleInstructions}");
-		static const FText QueryInstructionsAssistantFormat = LOCTEXT("QueryInstructionsFormat3",
+		static const FText QueryInstructionsAssistantFormat = INVTEXT(
 			"Provide an easily readable, informative, accurate answer that describes what I should know about yourself and what you can do for me. {StyleInstructions}");
-		static const FText QueryAdditionalInstructionsWidget = LOCTEXT("QueryAdditionalInstructions1",
+		static const FText QueryAdditionalInstructionsWidget = INVTEXT(
 			"Explain what happens if I select or click on it.");
-		static const FText QueryAdditionalInstructionsObject = LOCTEXT("QueryAdditionalInstructions2",
+		static const FText QueryAdditionalInstructionsObject = INVTEXT(
 			"Describe what the object is or what it does.");
-		static const FText QueryAdditionalInstructionsGeneral = LOCTEXT("QueryAdditionalInstructions3",
+		static const FText QueryAdditionalInstructionsGeneral = INVTEXT(
 			"If the object is clickable, like a button, then explain what happens if I click on it. If the object is something you select, like an asset or a blueprint node, describe what the object is or what it does.");
-		static const FText StyleInstructions = LOCTEXT("StyleInstructions1",
+		static const FText StyleInstructions = INVTEXT(
 			"Use bold text to emphasize Unreal Engine specific terms. Start with a quick overview paragraph, then add key points limited to one or two formatted sections with headers and bullet points. Don't include a summary at the end.");
 
 		if (!NameToQuery.IsEmpty())
@@ -1262,6 +1272,14 @@ void UE::AIAssistant::SlateQuerier::QueryAIAssistantAboutSlateWidget(const FWidg
 
 	// Construct structured context.
 	{
+		SlateQueryContext.StructuredContext.UnrealEditorVersion = FEngineVersion::Current().ToString();
+		SlateQueryContext.StructuredContext.UnrealEditorLanguage = FInternationalization::Get().GetCurrentLanguage()->GetDisplayName();
+		SlateQueryContext.StructuredContext.UnrealEditorLocale = FInternationalization::Get().GetCurrentLocale()->GetName();
+
+		static const FText LanguageInstructionsFormat = INVTEXT(" Please respond in language \"{0}\". ");
+		const FText Language = FText::FromString(SlateQueryContext.StructuredContext.UnrealEditorLanguage);
+		SlateQueryContext.LanguageInstructions = FText::Format(LanguageInstructionsFormat, Language);
+
 		SlateQueryContext.StructuredContext.ToolTipText = SlateQueryContext.CurrentToolTipText.ToString();
 
 		GenerateDetailsViewContext(InWidgetPath, SlateQueryContext);
@@ -1322,6 +1340,7 @@ void UE::AIAssistant::SlateQuerier::QueryAIAssistantAboutSlateWidget(const FWidg
 
 	const FString HiddenContextString = 
 		CleanExtraWhiteSpaceFromString(SlateQueryContext.GeneratedQueryInstructions.ToString()) +
+		SlateQueryContext.LanguageInstructions.ToString() +
 		TEXT(" ContextData=") +
 		SlateQueryContext.StructuredContext.ToJson(false);
 
