[AI Assistant]

Bind IME (Input Method Editor) to the web browser, for non-ASCII languages, like "Chinese (Traditional, Taiwan)".
(Yellow highlighting will be dealt with in a separate fix.)

diff --git a/Engine/Plugins/Experimental/AIAssistant/Source/AIAssistant/Private/AIAssistantWebBrowser.cpp b/Engine/Plugins/Experimental/AIAssistant/Source/AIAssistant/Private/AIAssistantWebBrowser.cpp
--- a/Engine/Plugins/Experimental/AIAssistant/Source/AIAssistant/Private/AIAssistantWebBrowser.cpp
+++ b/Engine/Plugins/Experimental/AIAssistant/Source/AIAssistant/Private/AIAssistantWebBrowser.cpp
@@ -188,6 +188,39 @@ void SAIAssistantWebBrowser::Construct(const FArguments& InArgs)
 		[
 			WebBrowserOrWebConnectionSwitcherWidget.ToSharedRef()
 		];
+
+
+	// Make the IME (Input Method Editor) work with HTML text input boxes on web pages,
+	// for non-ASCII languages, like "Chinese (Traditional, Taiwan)".
+	if (ITextInputMethodSystem* InputMethodSystem = FSlateApplication::Get().GetTextInputMethodSystem())
+    {
+		if (WebBrowserWidget.IsValid())
+		{
+			WebBrowserWidget->BindInputMethodSystem(InputMethodSystem);
+
+			// If UE is being shut down, we need to unbind the Input Method System before our
+			// destructor tries to do it. (Otherwise, during that time, Slate won't be in the
+			// correct state, and we can crash on shutdown.)
+			TWeakPtr<SWebBrowser> WebBrowserWidgetWeak = WebBrowserWidget; 
+			InputMethodSystemSlatePreShutdownDelegateHandle = FSlateApplication::Get().OnPreShutdown().AddLambda([WebBrowserWidgetWeak]() -> void
+			{
+				if (const TSharedPtr<SWebBrowser> WebBrowserWidgetLocal = WebBrowserWidgetWeak.Pin();
+					FSlateApplication::IsInitialized() && WebBrowserWidgetLocal.IsValid())
+				{
+					WebBrowserWidgetLocal->UnbindInputMethodSystem();
+				}
+			});
+		}
+    }
+    else
+    {
+		static const FText ErrorText = LOCTEXT(
+			"SAIAssistantWebBrowser_NonAsciiLanguageNotSupported",
+			"AI Assistant could not get an Input Method System. You will not be able to input "
+			"non-ASCII languages with AI Assistant.");
+        UE_LOG(LogAIAssistant, Error, TEXT("%hs() - %s"), __func__, *ErrorText.ToString());
+    }
+
 	
 #if !UE_AIA_SET_INITIAL_URL_VIA_WEB_BROWSER_SLATE_ARG
 	InitializeWebApplication();
@@ -195,13 +228,25 @@ void SAIAssistantWebBrowser::Construct(const FArguments& InArgs)
 #endif
 }
 
-SAIAssistantWebBrowser::~SAIAssistantWebBrowser()
+/*virtual*/ SAIAssistantWebBrowser::~SAIAssistantWebBrowser()
 {
 	InitializeWebApplication();
+				
+	// Unbind the Input Method Subsystem, and remove the Input Method Subsystem pre-shutdown lambda.
+	if (FSlateApplication::IsInitialized() && WebBrowserWidget.IsValid())
+	{
+		WebBrowserWidget->UnbindInputMethodSystem();
+
+		if (InputMethodSystemSlatePreShutdownDelegateHandle.IsValid())
+		{
+			FSlateApplication::Get().OnPreShutdown().Remove(InputMethodSystemSlatePreShutdownDelegateHandle);
+		}
+	}
 
+	// In case the web connection widget was trying to reconnect, stop it from reconnecting.
 	if (WebConnectionWidget.IsValid())
 	{
-		WebConnectionWidget->StopReconnecting(); // ..in case we were trying to reconnect
+		WebConnectionWidget->StopReconnecting(); 
 	}
 }
 
diff --git a/Engine/Plugins/Experimental/AIAssistant/Source/AIAssistant/Private/AIAssistantWebBrowser.h b/Engine/Plugins/Experimental/AIAssistant/Source/AIAssistant/Private/AIAssistantWebBrowser.h
--- a/Engine/Plugins/Experimental/AIAssistant/Source/AIAssistant/Private/AIAssistantWebBrowser.h
+++ b/Engine/Plugins/Experimental/AIAssistant/Source/AIAssistant/Private/AIAssistantWebBrowser.h
@@ -105,4 +105,6 @@ private:
 	mutable UE::FRecursiveMutex WebApplicationMutex;
 	// Web application used by the widget.
 	TSharedPtr<UE::AIAssistant::FWebApplication> MaybeWebApplication;
+
+	FDelegateHandle InputMethodSystemSlatePreShutdownDelegateHandle;
 };
