[AIAssistant]

Removed improperly gated WebBrowserWidget.IsValid() check for removing InputMethodSystemSlatePreShutdownDelegateHandle.

diff --git a/Engine/Plugins/Experimental/AIAssistant/Source/AIAssistant/Private/AIAssistantWebBrowser.cpp b/Engine/Plugins/Experimental/AIAssistant/Source/AIAssistant/Private/AIAssistantWebBrowser.cpp
--- a/Engine/Plugins/Experimental/AIAssistant/Source/AIAssistant/Private/AIAssistantWebBrowser.cpp
+++ b/Engine/Plugins/Experimental/AIAssistant/Source/AIAssistant/Private/AIAssistantWebBrowser.cpp
@@ -200,10 +200,11 @@ void SAIAssistantWebBrowser::Construct(const FArguments& InArgs)
 
 			// If UE is being shut down, we need to unbind the Input Method System before our
 			// destructor tries to do it. (Otherwise, during that time, Slate won't be in the
-			// correct state, and we can crash on shutdown.)
+			// correct state, and we can crash on shutdown.) 
 			TWeakPtr<SWebBrowser> WebBrowserWidgetWeak = WebBrowserWidget; 
 			InputMethodSystemSlatePreShutdownDelegateHandle = FSlateApplication::Get().OnPreShutdown().AddLambda([WebBrowserWidgetWeak]() -> void
 			{
+				// NOTE - Slate must be initialized or we can crash.
 				if (const TSharedPtr<SWebBrowser> WebBrowserWidgetLocal = WebBrowserWidgetWeak.Pin();
 					FSlateApplication::IsInitialized() && WebBrowserWidgetLocal.IsValid())
 				{
@@ -232,15 +233,18 @@ void SAIAssistantWebBrowser::Construct(const FArguments& InArgs)
 {
 	InitializeWebApplication();
 				
-	// Unbind the Input Method Subsystem, and remove the Input Method Subsystem pre-shutdown lambda.
+	// Unbind the Input Method Subsystem.
+	// NOTE - Slate must be initialized or we can crash.
 	if (FSlateApplication::IsInitialized() && WebBrowserWidget.IsValid())
 	{
 		WebBrowserWidget->UnbindInputMethodSystem();
-
-		if (InputMethodSystemSlatePreShutdownDelegateHandle.IsValid())
-		{
-			FSlateApplication::Get().OnPreShutdown().Remove(InputMethodSystemSlatePreShutdownDelegateHandle);
-		}
+	}
+	
+	// Remove the Input Method Subsystem pre-shutdown lambda.
+	// NOTE - Slate must be initialized or we can crash.
+	if (FSlateApplication::IsInitialized() && InputMethodSystemSlatePreShutdownDelegateHandle.IsValid())
+	{
+		FSlateApplication::Get().OnPreShutdown().Remove(InputMethodSystemSlatePreShutdownDelegateHandle);
 	}
 
 	// In case the web connection widget was trying to reconnect, stop it from reconnecting.
