[AI Assistant] Add Tooltip detection for Plugin browser

#rnx

diff --git a/Engine/Plugins/Experimental/AIAssistant/Source/AIAssistant/Private/AIAssistantSlateQuerier.cpp b/Engine/Plugins/Experimental/AIAssistant/Source/AIAssistant/Private/AIAssistantSlateQuerier.cpp
--- a/Engine/Plugins/Experimental/AIAssistant/Source/AIAssistant/Private/AIAssistantSlateQuerier.cpp
+++ b/Engine/Plugins/Experimental/AIAssistant/Source/AIAssistant/Private/AIAssistantSlateQuerier.cpp
@@ -598,6 +598,54 @@ static bool ItemIsDetailsProperty(const FWidgetPath& WidgetPath, FAIAssistantSla
 	return false;
 }
 
+// Return true and fill in SlateQueryContext if queried item is a plugin in Plugins browser.
+static bool ItemIsPlugin(const FWidgetPath& WidgetPath, FAIAssistantSlateQueryContext& SlateQueryContext)
+{
+	// Look up for containing SPluginBrowser.
+	TSharedPtr<SWidget> PluginBrowser = FindClosestWidgetPtrOfType(WidgetPath, "SPluginBrowser");
+	if (PluginBrowser.IsValid())
+	{
+		FText ChildText;
+		TSharedPtr<SWidget> PluginCategoryTree = FindClosestWidgetPtrOfType(WidgetPath, "SPluginCategoryTree");
+		TSharedPtr<SWidget> PluginTileList = FindClosestWidgetPtrOfType(WidgetPath, "SPluginTileList");
+		if (PluginCategoryTree.IsValid())
+		{
+			SlateQueryContext.LastPickedWidget = PluginCategoryTree;
+			TSharedPtr<SWidget> PluginCategory = FindClosestWidgetPtrOfType(WidgetPath, "SPluginCategory");
+			if (PluginCategory.IsValid())
+			{
+				SlateQueryContext.LastPickedWidget = PluginCategory;
+				ChildText = FindChildWidgetWithText(PluginCategory);
+				if (!ChildText.IsEmpty())
+				{
+					SlateQueryContext.ItemName = ChildText;
+					SlateQueryContext.ItemDescriptor = LOCTEXT("ItemDescriptor_PluginCategory", "plugin category");
+					SlateQueryContext.bIsObject = true;
+					return true;
+				}
+			}
+		}
+		else if (PluginTileList.IsValid())
+		{
+			SlateQueryContext.LastPickedWidget = PluginTileList;
+			TSharedPtr<SWidget> PluginTile = FindClosestWidgetPtrOfType(WidgetPath, "SPluginTile");
+			if (PluginTile.IsValid())
+			{
+				SlateQueryContext.LastPickedWidget = PluginTile;
+				ChildText = FindChildWidgetWithText(PluginTile);
+				if (!ChildText.IsEmpty())
+				{
+					SlateQueryContext.ItemName = ChildText;
+					SlateQueryContext.ItemDescriptor = LOCTEXT("ItemDescriptor_PluginTile", "plugin");
+					SlateQueryContext.bIsObject = true;
+					return true;
+				}
+			}
+		}
+	}
+	return false;
+}
+
 // Return true and fill in SlateQueryContext if queried item is a navigation breadcrumb button.
 static bool ItemIsBreadcrumbTrailButton(const FWidgetPath& WidgetPath, FAIAssistantSlateQueryContext& SlateQueryContext)
 {
@@ -763,6 +811,8 @@ static bool FindItemName(const FWidgetPath& WidgetPath, FAIAssistantSlateQueryCo
 		ItemIsConsoleInputBox(WidgetPath, SlateQueryContext) ||
 		// Is it a details panel property?
 		ItemIsDetailsProperty(WidgetPath, SlateQueryContext) ||
+		// Is it a plugin in the plugins browser ?
+		ItemIsPlugin(WidgetPath, SlateQueryContext) ||
 		// Is it a breadcrumb trail button?
 		ItemIsBreadcrumbTrailButton(WidgetPath, SlateQueryContext) ||
 		// Is it an asset item?
