[AIAssisstant]

Web connection widget 'no internet' border capped to not stretch beyond 800 pixels.

diff --git a/Engine/Plugins/Experimental/AIAssistant/Source/AIAssistant/Private/AIAssistantWebConnectionWidget.cpp b/Engine/Plugins/Experimental/AIAssistant/Source/AIAssistant/Private/AIAssistantWebConnectionWidget.cpp
--- a/Engine/Plugins/Experimental/AIAssistant/Source/AIAssistant/Private/AIAssistantWebConnectionWidget.cpp
+++ b/Engine/Plugins/Experimental/AIAssistant/Source/AIAssistant/Private/AIAssistantWebConnectionWidget.cpp
@@ -59,67 +59,73 @@ void SAIAssistantWebConnectionWidget::Construct(const FArguments& InArgs)
 						.VAlign(VAlign_Center)
 						.FillHeight(1.0f)
 						[
-							SNew(SBorder)
-								.Clipping(EWidgetClipping::ClipToBounds)
-								.BorderImage(FAppStyle::GetBrush("HomeScreen.NoInternet.Border"))
-								.Padding(16.0f)
+							SNew(SBox)
+								.HAlign(HAlign_Center)
+								.VAlign(VAlign_Center)
+								.MaxDesiredWidth(800.0f)
 								[
-									SNew(SVerticalBox)
-										+ SVerticalBox::Slot()
-											.AutoHeight()
-											.Padding(0.0f, 0.0f, 0.0f, 16.0f)
-											.HAlign(HAlign_Left)
-											[
-												SNew(SBox)
-													.WidthOverride(24.0f)
-													.HeightOverride(24.0f)
+									SNew(SBorder)
+										.Clipping(EWidgetClipping::ClipToBounds)
+										.BorderImage(FAppStyle::GetBrush("HomeScreen.NoInternet.Border"))
+										.Padding(16.0f)
+										[
+											SNew(SVerticalBox)
+												+ SVerticalBox::Slot()
+													.AutoHeight()
+													.Padding(0.0f, 0.0f, 0.0f, 16.0f)
+													.HAlign(HAlign_Left)
 													[
-														SNew(SWidgetSwitcher)
-															.WidgetIndex_Lambda([this, IsReconnectWaiting]() -> int32
-																{
-																	return (ConnectionState == EConnectionState::Reconnecting || IsReconnectWaiting() ? 1 : 0);
-																})
-															+ SWidgetSwitcher::Slot()
-																[
-																	SNew(SImage)
-																		.Image(FAppStyle::GetBrush("HomeScreen.NoInternet.Icon"))
-																]
-															+ SWidgetSwitcher::Slot()
-																[
-																	SNew(SCircularThrobber)
-																		.NumPieces(8)
-																		.Period(1.2f)
-																		.Radius(12.0f)
-																]
+														SNew(SBox)
+															.WidthOverride(24.0f)
+															.HeightOverride(24.0f)
+															[
+																SNew(SWidgetSwitcher)
+																	.WidgetIndex_Lambda([this, IsReconnectWaiting]() -> int32
+																		{
+																			return (ConnectionState == EConnectionState::Reconnecting || IsReconnectWaiting() ? 1 : 0);
+																		})
+																	+ SWidgetSwitcher::Slot()
+																		[
+																			SNew(SImage)
+																				.Image(FAppStyle::GetBrush("HomeScreen.NoInternet.Icon"))
+																		]
+																	+ SWidgetSwitcher::Slot()
+																		[
+																			SNew(SCircularThrobber)
+																				.NumPieces(8)
+																				.Period(1.2f)
+																				.Radius(12.0f)
+																		]
+															]
 													]
-											]
-											+ SVerticalBox::Slot()
-												.FillHeight(0.5f)
-												.Padding(0.0f, 0.0f, 0.0f, 16.0f)
-												[
-													SNew(STextBlock)
-														.Font(FAppStyle::GetFontStyle("BoldFont"))
-														.Text(InArgs._WhenDisconnectedMessage)
-														.ColorAndOpacity(FStyleColors::White)
-												]
-											+ SVerticalBox::Slot()
-												.FillHeight(1.0f)
-												[
-													SNew(SButton)
-														.HAlign(HAlign_Center)
-														.ButtonStyle(FAppStyle::Get(), "HomeScreen.MyFolderButton")
-														.Text(LOCTEXT("AIAssistantWebConnection_Reconnect", "Reconnect"))
-														.IsEnabled_Lambda([this, IsReconnectWaiting]() -> bool
-															{
-																return ((ConnectionState == EConnectionState::Disconnected) && !IsReconnectWaiting());
-															})
-														.OnClicked_Lambda([this]() -> FReply
-															{
-																StartReconnecting();
-																
-																return FReply::Handled();
-															})
-												]
+													+ SVerticalBox::Slot()
+														.FillHeight(0.5f)
+														.Padding(0.0f, 0.0f, 0.0f, 16.0f)
+														[
+															SNew(STextBlock)
+																.Font(FAppStyle::GetFontStyle("BoldFont"))
+																.Text(InArgs._WhenDisconnectedMessage)
+																.ColorAndOpacity(FStyleColors::White)
+														]
+													+ SVerticalBox::Slot()
+														.FillHeight(1.0f)
+														[
+															SNew(SButton)
+																.HAlign(HAlign_Center)
+																.ButtonStyle(FAppStyle::Get(), "HomeScreen.MyFolderButton")
+																.Text(LOCTEXT("AIAssistantWebConnection_Reconnect", "Reconnect"))
+																.IsEnabled_Lambda([this, IsReconnectWaiting]() -> bool
+																	{
+																		return ((ConnectionState == EConnectionState::Disconnected) && !IsReconnectWaiting());
+																	})
+																.OnClicked_Lambda([this]() -> FReply
+																	{
+																		StartReconnecting();
+																		
+																		return FReply::Handled();
+																	})
+														]
+										]
 								]
 						]
 			]
