Added clipping to the ai assistant widget.
This is a workaround, the web browser doesn't seem to respect our on demand clipping system unfortunatly.

diff --git a/Engine/Plugins/Experimental/AIAssistant/Source/AIAssistant/Private/AIAssistant.cpp b/Engine/Plugins/Experimental/AIAssistant/Source/AIAssistant/Private/AIAssistant.cpp
--- a/Engine/Plugins/Experimental/AIAssistant/Source/AIAssistant/Private/AIAssistant.cpp
+++ b/Engine/Plugins/Experimental/AIAssistant/Source/AIAssistant/Private/AIAssistant.cpp
@@ -302,7 +302,9 @@ void FAIAssistantModule::GenerateStatusBarPanelDrawerSummon(TArray<UStatusBarSub
 
 TSharedRef<SDockTab> FAIAssistantModule::OnSpawnPluginTab(const FSpawnTabArgs& SpawnTabArgs)
 {
-	AIAssistantWebBrowserWidget = SNew(SAIAssistantWebBrowser);
+	// Work around for the web browser drawing out of bounds while doing the on demand clipping
+	AIAssistantWebBrowserWidget = SNew(SAIAssistantWebBrowser)
+		.Clipping(EWidgetClipping::ClipToBounds);
 
 
 	TSharedRef<SDockTab> DockTabWidget = SNew(SDockTab)
