[AI Assistant] Fix issue where Tooltip hovering over asset thumbnail asks about selected actor in level editor instead

#rnx

diff --git a/Engine/Plugins/Experimental/AIAssistant/Source/AIAssistant/Private/AIAssistantSlateQuerier.cpp b/Engine/Plugins/Experimental/AIAssistant/Source/AIAssistant/Private/AIAssistantSlateQuerier.cpp
--- a/Engine/Plugins/Experimental/AIAssistant/Source/AIAssistant/Private/AIAssistantSlateQuerier.cpp
+++ b/Engine/Plugins/Experimental/AIAssistant/Source/AIAssistant/Private/AIAssistantSlateQuerier.cpp
@@ -801,9 +801,7 @@ static bool FindItemName(const FWidgetPath& WidgetPath, FAIAssistantSlateQueryCo
 		return false;
 	}
 
-	if (// Does the widget path end up in a Viewport?
-		ItemIsInViewport(WidgetPath, SlateQueryContext) ||
-		// Is it a graph node?
+	if (// Is it a graph node?
 		ItemIsGraphNode(WidgetPath, SlateQueryContext) ||
 		// Is it a search/filter field?
 		ItemIsSearchBox(WidgetPath, SlateQueryContext) ||
@@ -819,6 +817,8 @@ static bool FindItemName(const FWidgetPath& WidgetPath, FAIAssistantSlateQueryCo
 		ItemIsAsset(WidgetPath, SlateQueryContext) ||
 		// Is it an Actor in Outliner, or an instance or component in the SubobjectInstanceEditor?
 		ItemIsOutlinerActorOrSubobject(WidgetPath, SlateQueryContext) ||
+		// Does the widget path end up in a Viewport?
+		ItemIsInViewport(WidgetPath, SlateQueryContext) ||
 		// Is it one of the Details panel checkboxes that filter the detail properties viewed?
 		ItemIsDetailPanelFilterCheckbox(WidgetPath, SlateQueryContext) ||
 		// Is it a menu item?
