[AI Assistant] Fix for Tooltip crash during PIE

#rnx

diff --git a/Engine/Plugins/Experimental/AIAssistant/Source/AIAssistant/Private/AIAssistantSlateQuerier.cpp b/Engine/Plugins/Experimental/AIAssistant/Source/AIAssistant/Private/AIAssistantSlateQuerier.cpp
--- a/Engine/Plugins/Experimental/AIAssistant/Source/AIAssistant/Private/AIAssistantSlateQuerier.cpp
+++ b/Engine/Plugins/Experimental/AIAssistant/Source/AIAssistant/Private/AIAssistantSlateQuerier.cpp
@@ -940,40 +940,42 @@ static bool IsInEditor(const FWidgetPath& WidgetPath, FAIAssistantSlateQueryCont
 			SlateQueryContext.FormattedWindowName = FText::Format(LOCTEXT("EditorName_LevelEditorAndMode", " the Level Editor {WithModeActive}"), Args);
 		}
 
-		// Get selected actors, if any.
-		bool bActorIsSelected = false;
-		UEditorActorSubsystem* EditorActorSubsystem = GEditor->GetEditorSubsystem<UEditorActorSubsystem>();
-		// Get selected actors
-		TArray<AActor*> SelectedActors = EditorActorSubsystem->GetSelectedLevelActors();
-		for (AActor* Actor : SelectedActors)
+		if (!GEditor->IsPlaySessionInProgress())
 		{
-			if (Actor)
+			// Get selected actors, if any.
+			bool bActorIsSelected = false;
+			UEditorActorSubsystem* EditorActorSubsystem = GEditor->GetEditorSubsystem<UEditorActorSubsystem>();
+			// Get selected actors
+			TArray<AActor*> SelectedActors = EditorActorSubsystem->GetSelectedLevelActors();
+			for (AActor* Actor : SelectedActors)
 			{
-				SlateQueryContext.SelectedActorName = Actor->GetName();
-				SlateQueryContext.SelectedActorClass = Actor->GetClass()->GetName();
-				bActorIsSelected = true;
-				break;
+				if (Actor)
+				{
+					SlateQueryContext.SelectedActorName = Actor->GetName();
+					SlateQueryContext.SelectedActorClass = Actor->GetClass()->GetName();
+					bActorIsSelected = true;
+					break;
+				}
 			}
-		}
-		if (!bActorIsSelected)
-		{
-			FViewport* ActiveViewport = GEditor->GetActiveViewport();
-			if (ActiveViewport)
+
+			// If there is no selection and mouse is in viewport, see what actor is under mouse cursor.
+			if (!bActorIsSelected && FindClosestWidgetPtrOfType(WidgetPath, "SViewport"))
 			{
-				FIntPoint MousePos;
-				ActiveViewport->GetMousePos(MousePos);
-				HHitProxy* HitProxy = ActiveViewport->GetHitProxy(MousePos.X, MousePos.Y);
-				if (HitProxy)
+				FViewport* ActiveViewport = GEditor->GetActiveViewport();
+				if (ActiveViewport)
 				{
-					if (HitProxy->IsA(HActor::StaticGetType()))
+					FIntPoint MousePos;
+					ActiveViewport->GetMousePos(MousePos);
+					HHitProxy* HitProxy = ActiveViewport->GetHitProxy(MousePos.X, MousePos.Y);
+					if (HitProxy && HitProxy->IsA(HActor::StaticGetType()))
 					{
 						HActor* TargetProxy = static_cast<HActor*>(HitProxy);
 						AActor* TargetActor = TargetProxy->Actor;
-						// Double check for not null
 						if (TargetActor)
 						{
 							SlateQueryContext.SelectedActorName = TargetActor->GetName();
 							SlateQueryContext.SelectedActorClass = TargetActor->GetClass()->GetName();
+							bActorIsSelected = true;
 						}
 					}
 				}
