[AIAssistant]

Fix assert on Mac in our SWebBrowser's OnLoadCompleted_Lambda().
Run the 'set keyboard focus' code on the Game Thread, so FSlateApplication::Get() doesn't assert.
On Mac, the lamba will be run on an AppKit thread, which causes the assert.

diff --git a/Engine/Plugins/Experimental/AIAssistant/Source/AIAssistant/Private/AIAssistantWebBrowser.cpp b/Engine/Plugins/Experimental/AIAssistant/Source/AIAssistant/Private/AIAssistantWebBrowser.cpp
--- a/Engine/Plugins/Experimental/AIAssistant/Source/AIAssistant/Private/AIAssistantWebBrowser.cpp
+++ b/Engine/Plugins/Experimental/AIAssistant/Source/AIAssistant/Private/AIAssistantWebBrowser.cpp
@@ -8,6 +8,7 @@
 #include "IWebBrowserWindow.h"
 #include "Misc/AssertionMacros.h"
 #include "Framework/Application/SlateApplication.h"
+#include "Async/TaskGraphInterfaces.h"
 
 #include "AIAssistantConfig.h"
 #include "AIAssistantLog.h"
@@ -154,11 +155,20 @@ void SAIAssistantWebBrowser::Construct(const FArguments& InArgs)
 
 			// Set keyboard focus. This makes sure that a blinking cursor will appear in text
 			// boxes that are clicked on. Without this, this isn't always the case.
-			// IMPORTANT - Check validity, because we can rarely end up here during Editor shutdown.
-			if (WebBrowserWidget.IsValid() && FSlateApplication::IsInitialized())
+			// NOTE - Make sure we run this in Game Thread, so FSlateApplication::Get() doesn't
+			// assert. On Mac, the lambda we're in now will often get called from an AppKit thread.
 			{
-				FSlateApplication::Get().SetKeyboardFocus(
-					WebBrowserWidget, EFocusCause::SetDirectly);
+				TWeakPtr<SWebBrowser> WebBrowserWidgetWeak = WebBrowserWidget;
+				AsyncTask(ENamedThreads::GameThread, [WebBrowserWidgetWeak]() -> void
+					{
+						// Safety checks, since we can rarely end up here during Editor shutdown.
+						if (const TSharedPtr<SWebBrowser> WebBrowserWidget = WebBrowserWidgetWeak.Pin();
+							WebBrowserWidget.IsValid() && FSlateApplication::IsInitialized())
+						{
+							FSlateApplication::Get().SetKeyboardFocus(
+								WebBrowserWidget, EFocusCause::SetDirectly);
+						}
+					});
 			}
 		});
 	
