[AI Assistant] Add more UI query heuristics to detect:
- SFilterSearchBox
- SConsoleInputBox
- Asset items (folders, thumbnails) in content browser
- drawer overlays containing ContentBrowser, OutputLog
- breadcrumb folder navigation

Also fix crash when widget path is empty

diff --git a/Engine/Plugins/Experimental/AIAssistant/Source/AIAssistant/Private/AIAssistantSlateQuerier.cpp b/Engine/Plugins/Experimental/AIAssistant/Source/AIAssistant/Private/AIAssistantSlateQuerier.cpp
--- a/Engine/Plugins/Experimental/AIAssistant/Source/AIAssistant/Private/AIAssistantSlateQuerier.cpp
+++ b/Engine/Plugins/Experimental/AIAssistant/Source/AIAssistant/Private/AIAssistantSlateQuerier.cpp
@@ -13,8 +13,10 @@
 #include "Toolkits/BaseToolkit.h"
 #include "Widgets/Text/STextBlock.h"
 #include "Widgets/Text/SRichTextBlock.h"
+#include "Widgets/Input/SMultiLineEditableTextBox.h"
 #include "Widgets/Input/SSearchBox.h"
 #include "Widgets/Docking/SDockTab.h"
+#include "Filters/SFilterSearchBox.h"
 #include "SGraphNode.h"
 #include "Framework/Application/SlateApplication.h"
 
@@ -177,30 +179,33 @@ static FText FindCurrentToolTipText()
 
 static TSharedRef<SWidget> FindClosestMenuItem(const FWidgetPath& WidgetPath)
 {
-	// is it a menu item? Look up for an SMenuEntryBlock or SWidgetBlock.
-	TSharedRef<SWidget> MenuItemBlock = FindClosestWidgetOfType(WidgetPath, "SMenuEntryBlock");
-	if (MenuItemBlock->GetType() == "SNullWidgetContent")
+	if (WidgetPath.IsValid())
 	{
-		MenuItemBlock = FindClosestWidgetOfType(WidgetPath, "SWidgetBlock");
-	}
-	// disabled menu items *contain* the menu entry block
-	if (MenuItemBlock->GetType() == "SNullWidgetContent")
-	{
-		if (WidgetPath.GetLastWidget()->GetChildren()->Num() > 0)
+		// is it a menu item? Look up for an SMenuEntryBlock or SWidgetBlock.
+		TSharedRef<SWidget> MenuItemBlock = FindClosestWidgetOfType(WidgetPath, "SMenuEntryBlock");
+		if (MenuItemBlock->GetType() == "SNullWidgetContent")
 		{
-			for (int32 ChildIdx = 0; ChildIdx < WidgetPath.GetLastWidget()->GetChildren()->Num(); ChildIdx++)
+			MenuItemBlock = FindClosestWidgetOfType(WidgetPath, "SWidgetBlock");
+		}
+		// disabled menu items *contain* the menu entry block
+		if (MenuItemBlock->GetType() == "SNullWidgetContent")
+		{
+			if (WidgetPath.GetLastWidget()->GetChildren()->Num() > 0)
 			{
-				TSharedRef<SWidget> ThisWidget = WidgetPath.GetLastWidget()->GetChildren()->GetChildAt(ChildIdx);
-				if (ThisWidget->GetType() == "SMenuEntryBlock" || ThisWidget->GetType() == "SWidgetBlock")
+				for (int32 ChildIdx = 0; ChildIdx < WidgetPath.GetLastWidget()->GetChildren()->Num(); ChildIdx++)
 				{
-					MenuItemBlock = ThisWidget;
+					TSharedRef<SWidget> ThisWidget = WidgetPath.GetLastWidget()->GetChildren()->GetChildAt(ChildIdx);
+					if (ThisWidget->GetType() == "SMenuEntryBlock" || ThisWidget->GetType() == "SWidgetBlock")
+					{
+						MenuItemBlock = ThisWidget;
+					}
 				}
 			}
 		}
-	}
-	if (MenuItemBlock->GetType() != "SNullWidgetContent")
-	{
-		return MenuItemBlock;
+		if (MenuItemBlock->GetType() != "SNullWidgetContent")
+		{
+			return MenuItemBlock;
+		}
 	}
 
 	return SNullWidget::NullWidget;
@@ -272,7 +277,12 @@ static FText FindItemName(const FWidgetPath& WidgetPath, FAIAssistantSlateQueryC
 		});
 
 	FText ChildText = FText();
-	
+
+	if (!WidgetPath.IsValid())
+	{
+		return OutName;
+	}
+
 	// Is it a graph node?
 	// Look up for an SGraphEditor and then down for an SGraphPanel. Whatever is the child of that SGraphPanel is our SGraphNode.
 	// Note that this rests on the assumption that SGraphPanel only contains children that are SGraphNodes. This assumption is also
@@ -369,23 +379,52 @@ static FText FindItemName(const FWidgetPath& WidgetPath, FAIAssistantSlateQueryC
 		if (!ChildText.IsEmpty() || !FindCurrentToolTipText().IsEmpty())
 		{
 			OutName = (!ChildText.IsEmpty()) ? ChildText : LOCTEXT("ButtonWithNoTextPlaceholderName", "[icon-only]");
-			OutDescriptor = LOCTEXT("ItemDescriptor_Menu", "toolbar button");
+			OutDescriptor = LOCTEXT("ItemDescriptor_ToolbarButton", "toolbar button");
 			SlateQueryContext.LastPickedWidget = ToolbarItemBlock.ToSharedPtr();
 		}
 	}
 
 
 	// Is it a search/filter field?
+	FText HintText;
 	TSharedRef<SWidget> SearchBox = FindClosestWidgetOfType(WidgetPath, "SSearchBox");
 	if (SearchBox->GetType() != "SNullWidgetContent")
 	{
 		TSharedRef<SSearchBox> ThisSearchBoxCast = StaticCastSharedRef<SSearchBox>(SearchBox);
-		FText HintText = ThisSearchBoxCast->GetHintText();
+		HintText = ThisSearchBoxCast->GetHintText();
+	}
+	else
+	{
+		SearchBox = FindClosestWidgetOfType(WidgetPath, "SFilterSearchBox");
+		if (SearchBox->GetType() != "SNullWidgetContent")
+		{
+			TSharedRef<SWidget> EditableText = FindChildWidgetOfType(SearchBox, "SEditableText");
+			if (EditableText->GetType() != "SNullWidgetContent")
+			{
+				TSharedRef<SEditableText> ThisEditableTextCast = StaticCastSharedRef<SEditableText>(EditableText);
+				HintText = ThisEditableTextCast->GetHintText();
+			}
+		}
+	}
+	if (SearchBox->GetType() != "SNullWidgetContent")
+	{
 		OutName = FText::Format(LOCTEXT("SearchBoxFormat", "{0}"), (HintText.IsEmpty()) ? LOCTEXT("SearchBoxDefaultName","search") : HintText);
-		OutDescriptor = LOCTEXT("ItemDescriptor_Menu", " box");
+		OutDescriptor = LOCTEXT("ItemDescriptor_SearchBox", " box");
 		SlateQueryContext.LastPickedWidget = SearchBox.ToSharedPtr();
 	}
 
+	// Is it a console input box?
+	TSharedRef<SWidget> ConsoleInputBox = FindClosestWidgetOfType(WidgetPath, "SConsoleInputBox");
+	if (ConsoleInputBox->GetType() != "SNullWidgetContent")
+	{
+		TSharedRef<SWidget> EditableText = FindChildWidgetOfType(ConsoleInputBox, "SMultiLineEditableTextBox");
+		TSharedRef<SMultiLineEditableTextBox> ThisMultiLineEditableTextBoxCast = StaticCastSharedRef<SMultiLineEditableTextBox>(EditableText);
+		HintText = ThisMultiLineEditableTextBoxCast->GetHintText();
+		OutName = FText::Format(LOCTEXT("ConsoleInputBoxFormat", "{0}"), (HintText.IsEmpty() ? LOCTEXT("ConsoleInputBoxDefaultName","Console") : HintText));
+		OutDescriptor = LOCTEXT("ItemDescriptor_ConsoleInputBox", " input box");
+		SlateQueryContext.LastPickedWidget = ConsoleInputBox.ToSharedPtr();
+	}
+
 	// Is it a menu item? Look up for an SMenuEntryBlock or SWidgetBlock and then down for the text.
 	TSharedRef<SWidget> MenuItemBlock = FindClosestMenuItem(WidgetPath);
 	if (MenuItemBlock->GetType() != "SNullWidgetContent")
@@ -416,6 +455,48 @@ static FText FindItemName(const FWidgetPath& WidgetPath, FAIAssistantSlateQueryC
 		}
 	}
 
+	// Is it a breadcrumb trail button?
+	TSharedRef<SWidget> BreadcrumbTrail = FindClosestWidgetOfType(WidgetPath, "SBreadcrumbTrail<FNavigationCrumb>");
+	if (BreadcrumbTrail->GetType() != "SNullWidgetContent")
+	{
+		TSharedRef<SWidget> BreadcrumbButton = FindClosestWidgetOfType(WidgetPath, "SButton");
+		ChildText = FindChildWidgetWithText(BreadcrumbButton);
+		if (!ChildText.IsEmpty())
+		{
+			OutName = ChildText;
+			OutDescriptor = LOCTEXT("ItemDescriptor_Breadcrumb", "navigation breadcrumb");
+			SlateQueryContext.LastPickedWidget = BreadcrumbButton.ToSharedPtr();
+		}
+	}
+
+	// Is it an asset tile item?
+	TSharedRef<SWidget> AssetTileItem = FindClosestWidgetOfType(WidgetPath, "SAssetTileItem");
+	if (AssetTileItem->GetType() != "SNullWidgetContent")
+	{
+		// TSharedRef<SAssetTileItem> ThisAssetTileItemCast = StaticCastSharedRef<SAssetTileItem>(AssetTileItem);
+		TSharedRef<SWidget> AssetThumbnail = FindClosestWidgetOfType(WidgetPath, "SAssetThumbnail");
+		if (AssetThumbnail->GetType() != "SNullWidgetContent")
+		{
+			ChildText = FindChildWidgetWithText(AssetThumbnail); // This returns TYPE of asset instead of name
+			OutDescriptor = LOCTEXT("ItemDescriptor_Asset", "asset");
+			if (!ChildText.IsEmpty())
+			{
+				OutName = ChildText;
+				SlateQueryContext.LastPickedWidget = AssetThumbnail.ToSharedPtr();
+			}
+		}
+		else
+		{
+			ChildText = FindChildWidgetWithText(AssetTileItem);
+			OutDescriptor = LOCTEXT("ItemDescriptor_Folder", "asset folder");
+			if (!ChildText.IsEmpty())
+			{
+				OutName = ChildText;
+				SlateQueryContext.LastPickedWidget = AssetTileItem.ToSharedPtr();
+			}
+		}
+	}
+
 	if (!OutName.IsEmpty())
 	{
 		return FText::Format(LOCTEXT("ItemName_Format", " the \"{0}\" {1}"), { OutName, OutDescriptor });
@@ -491,6 +572,11 @@ static FText FindEditorName(const FWidgetPath& WidgetPath, FAIAssistantSlateQuer
 {
 	FText OutName = FText();
 
+	if (!WidgetPath.IsValid())
+	{
+		return OutName;
+	}
+
 	// is it in the status bar?
 	TSharedRef<SWidget> StatusBar = FindClosestWidgetOfType(WidgetPath, "SStatusBar");
 	if (StatusBar->GetType() != "SNullWidgetContent")
@@ -561,6 +647,34 @@ static FText FindEditorName(const FWidgetPath& WidgetPath, FAIAssistantSlateQuer
 		}
 	}
 
+	// Is it in a drawer overlay?
+	TSharedRef<SWidget> DrawerOverlay = FindFirstWidgetOfType(WidgetPath, "SDrawerOverlay");
+	if (DrawerOverlay->GetType() != "SNullWidgetContent")
+	{
+		FText DrawerName;
+
+		TSharedRef<SWidget> DrawerWidget = FindChildWidgetOfType(DrawerOverlay, "SContentBrowser");
+		if (DrawerWidget->GetType() != "SNullWidgetContent")
+		{
+			DrawerName = LOCTEXT("DrawerName_ContentBrowser", "ContentBrowser");
+		}
+		else
+		{
+			DrawerWidget = FindChildWidgetOfType(DrawerOverlay, "SOutputLog");
+			if (DrawerWidget->GetType() != "SNullWidgetContent")
+			{
+				DrawerName = LOCTEXT("DrawerName_OutputLog", "OutputLog");
+			}
+		}
+		// Are there other items that are commonly in drawers?
+
+		if (!DrawerName.IsEmpty())
+		{
+			SlateQueryContext.LastPickedWidget = DrawerWidget.ToSharedPtr();
+			return FText::Format(LOCTEXT("DrawerNameFormat", "the {0} drawer"), DrawerName);
+		}
+	}
+
 	TSharedRef<SWindow> WidgetWindow = WidgetPath.GetWindow();
 	// is it the main window that hosts the level editor?
 	if (FModuleManager::Get().IsModuleLoaded("LevelEditor"))
@@ -589,31 +703,34 @@ static FText FindTextUnderCursor(const FWidgetPath& WidgetPath, FAIAssistantSlat
 {
 	FText OutText;
 
-	const FRegexPattern AlphaPattern(TEXT("[A-Za-z]+"));
-	TSharedRef<SWidget> WidgetToTest = WidgetPath.GetLastWidget();
-	FText WidgetText;
-	if (WidgetToTest->GetType() == "STextBlock")
-	{
-		TSharedRef<STextBlock> TextBlock = StaticCastSharedRef<STextBlock>(WidgetToTest);
-		WidgetText = TextBlock->GetText();
-	}
-	else if (WidgetToTest->GetType() == "SRichTextBlock")
-	{
-		TSharedRef<SRichTextBlock> TextBlock = StaticCastSharedRef<SRichTextBlock>(WidgetToTest);
-		WidgetText = TextBlock->GetText();
-	}
-	if (!WidgetText.IsEmpty())
+	if (WidgetPath.IsValid())
 	{
-		FRegexMatcher AlphaMatcher(AlphaPattern, WidgetText.ToString());
-		if (AlphaMatcher.FindNext())
+		const FRegexPattern AlphaPattern(TEXT("[A-Za-z]+"));
+		TSharedRef<SWidget> WidgetToTest = WidgetPath.GetLastWidget();
+		FText WidgetText;
+		if (WidgetToTest->GetType() == "STextBlock")
 		{
-			OutText = WidgetText;
+			TSharedRef<STextBlock> TextBlock = StaticCastSharedRef<STextBlock>(WidgetToTest);
+			WidgetText = TextBlock->GetText();
+		}
+		else if (WidgetToTest->GetType() == "SRichTextBlock")
+		{
+			TSharedRef<SRichTextBlock> TextBlock = StaticCastSharedRef<SRichTextBlock>(WidgetToTest);
+			WidgetText = TextBlock->GetText();
+		}
+		if (!WidgetText.IsEmpty())
+		{
+			FRegexMatcher AlphaMatcher(AlphaPattern, WidgetText.ToString());
+			if (AlphaMatcher.FindNext())
+			{
+				OutText = WidgetText;
+			}
 		}
-	}
 
-	if (!OutText.IsEmpty())
-	{
-		AccumulatedSlateQueryContext.LastPickedWidget = WidgetToTest;
+		if (!OutText.IsEmpty())
+		{
+			AccumulatedSlateQueryContext.LastPickedWidget = WidgetToTest;
+		}
 	}
 
 	return OutText;
@@ -686,7 +803,7 @@ void UE::AIAssistant::SlateQuerier::QueryAIAssistantAboutSlateWidgetUnderCursor(
 	// Is there a current tool-tip?
 	const FText ToolTipText = FindCurrentToolTipText();
 
-	
+
 	{
 		static const FText QueryAdditionalInstructions = LOCTEXT("QueryInstructions",
 			"If the object is clickable, like a button, then explain what happens if I click on it, including any unrecoverable or unexpected side effects.  If it is safe, don't bother mentioning it. If the object is something you select, like an asset or a blueprint node, describe what the object is or what it does.");
@@ -764,35 +881,30 @@ void UE::AIAssistant::SlateQuerier::QueryAIAssistantAboutSlateWidgetUnderCursor(
 		}
 	}
 
-
-	// OPTIONAL - We're done using these members. We can clear them to reduce size for lamba capture below.
+	// OPTIONAL - We're done using these members. We can clear them to reduce size for lambda capture below.
 	
 	SlateQueryContext.LastPickedWidget.Reset();
 	SlateQueryContext.GeneratedContextItems.Empty();
 
-	
+
 	// Send widget query to AI Assistant.
 
-	if (!AIAssistantWebBrowserWidget->IsExecuteWhenReadyPending()) // ..avoid user prompt flooding
-	{
-		AIAssistantWebBrowserWidget->ExecuteWhenReady([SlateQueryContext]
-			{
-				UAIAssistantSubsystem::NewChatThread();
-
-					
-				const FString VisiblePromptString =
-					CleanExtraWhiteSpaceFromString(SlateQueryContext.GeneratedQuery.ToString());
-
-				const FString HiddenContextString = 
-					CleanExtraWhiteSpaceFromString(SlateQueryContext.GeneratedQueryInstructions.ToString()) +
-					FText(LOCTEXT("ContextPrefix", "(Context: ")).ToString() +
-					CleanExtraWhiteSpaceFromString(SlateQueryContext.GeneratedContext.ToString()) +
-					FText(LOCTEXT("ContextPostfix", ")")).ToString();
-					
-				UAIAssistantSubsystem::SendUserPrompt(VisiblePromptString, HiddenContextString);
+	AIAssistantWebBrowserWidget->ExecuteWhenReady([SlateQueryContext]
+		{
+			UAIAssistantSubsystem::NewChatThread();
+
 				
-			});
-	}
+			const FString VisiblePromptString =
+				CleanExtraWhiteSpaceFromString(SlateQueryContext.GeneratedQuery.ToString());
+
+			const FString HiddenContextString = 
+				CleanExtraWhiteSpaceFromString(SlateQueryContext.GeneratedQueryInstructions.ToString()) +
+				FText(LOCTEXT("ContextPrefix", "(Context: ")).ToString() +
+				CleanExtraWhiteSpaceFromString(SlateQueryContext.GeneratedContext.ToString()) +
+				FText(LOCTEXT("ContextPostfix", ")")).ToString();
+				
+			UAIAssistantSubsystem::SendUserPrompt(VisiblePromptString, HiddenContextString);
+		});
 }
 
 
