Fix URL tracking in navigation policy

When opening a URL in the system browser the navigation policy should not log
the URL as being opened.

diff --git a/Engine/Plugins/Experimental/AIAssistant/Source/AIAssistant/Private/AIAssistantWebBrowserNavigationPolicy.cpp b/Engine/Plugins/Experimental/AIAssistant/Source/AIAssistant/Private/AIAssistantWebBrowserNavigationPolicy.cpp
--- a/Engine/Plugins/Experimental/AIAssistant/Source/AIAssistant/Private/AIAssistantWebBrowserNavigationPolicy.cpp
+++ b/Engine/Plugins/Experimental/AIAssistant/Source/AIAssistant/Private/AIAssistantWebBrowserNavigationPolicy.cpp
@@ -19,24 +19,30 @@ namespace UE::AIAssistant
 		const FString& Url, const FWebNavigationRequest& Request)
 	{
 		bOpenedUrlChanged = Url != LastOpenedUrl;
-		LastOpenedUrl = Url;
 
 		EOpenBrowserMode OpenBrowserMode =
-			IsCurrentUrlMainUrl()
+			IsMainUrl(Url)
 			? EOpenBrowserMode::Embedded
 			: Config->OpenBrowserModeFromNavigation(
 				Url, Request.bIsRedirect, LastOpenedNonRedirectUrl);
 
-		// If the page can be opened in the embedded browser keep track of the last non-redirect URL.
-		if (OpenBrowserMode == EOpenBrowserMode::Embedded && !Request.bIsRedirect)
+		if (OpenBrowserMode == EOpenBrowserMode::Embedded)
 		{
-			LastOpenedNonRedirectUrl = Url;
+			LastOpenedUrl = Url;
+			// If the page can be opened in the embedded browser keep track of the last
+			// non-redirect URL.
+			if (!Request.bIsRedirect) LastOpenedNonRedirectUrl = Url;
 		}
 		return OpenBrowserMode;
 	}
 
 	bool FWebBrowserNavigationPolicy::IsCurrentUrlMainUrl() const
 	{
-		return !Config->MainUrl.IsEmpty() && LastOpenedUrl.StartsWith(Config->MainUrl);
+		return IsMainUrl(LastOpenedUrl);
+	}
+
+	bool FWebBrowserNavigationPolicy::IsMainUrl(const FString& Url) const
+	{
+		return !Config->MainUrl.IsEmpty() && Url.StartsWith(Config->MainUrl);
 	}
 }
\ No newline at end of file
diff --git a/Engine/Plugins/Experimental/AIAssistant/Source/AIAssistant/Private/AIAssistantWebBrowserNavigationPolicy.h b/Engine/Plugins/Experimental/AIAssistant/Source/AIAssistant/Private/AIAssistantWebBrowserNavigationPolicy.h
--- a/Engine/Plugins/Experimental/AIAssistant/Source/AIAssistant/Private/AIAssistantWebBrowserNavigationPolicy.h
+++ b/Engine/Plugins/Experimental/AIAssistant/Source/AIAssistant/Private/AIAssistantWebBrowserNavigationPolicy.h
@@ -22,9 +22,13 @@ namespace UE::AIAssistant
 		// Has the opened URL changed since the last call to update.
 		bool HasOpenedUrlChanged() const { return bOpenedUrlChanged; }
 
-		// Whether the current URL is the main url in the config.
+		// Whether the current URL is the main URL in the config.
 		bool IsCurrentUrlMainUrl() const;
 
+	private:
+		// Whether the specified URL is the main URL in the config.
+		bool IsMainUrl(const FString& Url) const;
+
 	private:
 		TSharedRef<FWebBrowserConfig> Config;
 		// Previously opened URL.
diff --git a/Engine/Plugins/Experimental/AIAssistant/Source/AIAssistant/Private/Tests/AIAssistantWebBrowserNavigationPolicyTest.cpp b/Engine/Plugins/Experimental/AIAssistant/Source/AIAssistant/Private/Tests/AIAssistantWebBrowserNavigationPolicyTest.cpp
--- a/Engine/Plugins/Experimental/AIAssistant/Source/AIAssistant/Private/Tests/AIAssistantWebBrowserNavigationPolicyTest.cpp
+++ b/Engine/Plugins/Experimental/AIAssistant/Source/AIAssistant/Private/Tests/AIAssistantWebBrowserNavigationPolicyTest.cpp
@@ -118,6 +118,7 @@ bool FAIAssistantWebBrowserNavigationPolicyTestUpdateToEmbeddedRedirect::RunTest
 		TEXT("UseEmbeddedBrowser"),
 		NavigationPolicy.Update(TEXT("https://a.random.url"), CreateWebNavigationRequestForRedirect()),
 		EOpenBrowserMode::Embedded);
+	(void)TestFalse(TEXT("IsNotMainUrl"), NavigationPolicy.IsCurrentUrlMainUrl());
 	return true;
 }
 
@@ -138,6 +139,7 @@ bool FAIAssistantWebBrowserNavigationPolicyTestUpdateToSystem::RunTest(
 			TEXT("https://another.random.url"),
 			CreateWebNavigationRequest()),
 		EOpenBrowserMode::System);
+	(void)TestTrue(TEXT("IsStillMainUrl"), NavigationPolicy.IsCurrentUrlMainUrl());
 	return true;
 }
 
