Added IToolsetHandler and a Toolset registry, along with functions to register, and unregister toolsets, and execute tool calls.

diff --git a/Engine/Plugins/Experimental/AIAssistant/Source/AIAssistant/Private/AIAssistantSubsystem.h b/Engine/Plugins/Experimental/AIAssistant/Source/AIAssistant/Private/AIAssistantSubsystem.h
--- a/Engine/Plugins/Experimental/AIAssistant/Source/AIAssistant/Private/AIAssistantSubsystem.h
+++ b/Engine/Plugins/Experimental/AIAssistant/Source/AIAssistant/Private/AIAssistantSubsystem.h
@@ -3,7 +3,9 @@
 
 #pragma once
 
+#include "Containers/UnrealString.h"
 #include "EditorSubsystem.h"
+#include "AIAssistantToolsetRegistry.h"
 #include "Templates/SharedPointer.h"
 
 #include "AIAssistantSubsystem.generated.h"
@@ -43,4 +45,7 @@ public:
 	static FAIAssistantModule& GetAIAssistantModule();
 	// Get the assistant web browser.
 	static TSharedPtr<SAIAssistantWebBrowser> GetAIAssistantWebBrowserWidget();
+
+	// The set of registered toolset handlers.
+	UE::AIAssistant::FToolsetRegistry ToolsetRegistry;
 };
diff --git a/Engine/Plugins/Experimental/AIAssistant/Source/AIAssistant/Private/AIAssistantToolsetHandler.h b/Engine/Plugins/Experimental/AIAssistant/Source/AIAssistant/Private/AIAssistantToolsetHandler.h
new file mode 100644
--- /dev/null
+++ b/Engine/Plugins/Experimental/AIAssistant/Source/AIAssistant/Private/AIAssistantToolsetHandler.h
@@ -0,0 +1,31 @@
+// Copyright Epic Games, Inc. All Rights Reserved.
+
+#pragma once
+
+#include "Async/Future.h"
+#include "Containers/UnrealString.h"
+#include "Dom/JsonObject.h"
+#include "Templates/ValueOrError.h"
+
+namespace UE::AIAssistant {
+	// Universal interface that all toolsets must implement.
+	class IToolsetHandler {
+	public:
+		virtual ~IToolsetHandler() = default;
+
+		// Executes a tool with the given name and JSON input, returning a future with the JSON
+		// output.
+		virtual TFuture<TValueOrError<FJsonObject, FString>> ExecuteTool(
+			const FString& ToolName, const FJsonObject& JsonInput) = 0;
+
+		// Returns the JSON schema for the toolset
+		virtual FJsonObject GetToolsetJsonSchema() const = 0;
+
+		// Returns the name of the toolset
+		virtual FString GetToolsetName() const = 0;
+
+		// Returns the version of the toolset
+		virtual FString GetToolsetVersion() const = 0;
+	};
+
+}
diff --git a/Engine/Plugins/Experimental/AIAssistant/Source/AIAssistant/Private/AIAssistantToolsetRegistry.cpp b/Engine/Plugins/Experimental/AIAssistant/Source/AIAssistant/Private/AIAssistantToolsetRegistry.cpp
new file mode 100644
--- /dev/null
+++ b/Engine/Plugins/Experimental/AIAssistant/Source/AIAssistant/Private/AIAssistantToolsetRegistry.cpp
@@ -0,0 +1,130 @@
+// Copyright Epic Games, Inc. All Rights Reserved.
+
+#include "AIAssistantToolsetRegistry.h"
+
+#include "Containers/Array.h"
+#include "Containers/UnrealString.h"
+#include "Editor.h"
+#include "Logging/LogVerbosity.h"
+#include "Misc/Optional.h"
+#include "Misc/Variant.h"
+#include "Modules/ModuleManager.h"
+#include "Templates/SharedPointer.h"
+
+#include "AIAssistant.h"
+#include "AIAssistantPythonExecutor.h"
+#include "AIAssistantToolsetHandler.h"
+#include "AIAssistantWebBrowser.h"
+
+namespace UE::AIAssistant
+{
+	TValueOrError<FToolDescriptor, FString> FToolDescriptor::FromString(
+		const FString& ToolsetAndToolName, bool LogErrors)
+	{
+		FToolDescriptor Descriptor;
+		if (!ToolsetAndToolName.Split(
+			TEXT("."), &Descriptor.ToolsetName, &Descriptor.ToolName,
+			ESearchCase::IgnoreCase, ESearchDir::FromEnd) ||
+			Descriptor.ToolsetName.IsEmpty() || Descriptor.ToolName.IsEmpty())
+		{
+			FString Message =
+				TEXT("Invalid tool name format. Expected format: ToolsetName.ToolName");
+			if (LogErrors)
+			{
+				UE_LOG(LogAIAssistant, Error, TEXT("%s"), *Message);
+			}
+			return MakeError(Message);
+		}
+		return MakeValue(Descriptor);
+	}
+
+	bool FToolsetRegistry::RegisterToolset(TSharedPtr<IToolsetHandler> ToolsetHandler)
+	{
+		if (!ToolsetHandler.IsValid())
+		{
+			UE_LOG(LogAIAssistant, Warning, TEXT("Toolset must not be null"));
+			return false;
+		}
+		if (ToolsetHandlers.Contains(ToolsetHandler->GetToolsetName()))
+		{
+			UE_LOG(
+				LogAIAssistant, Warning, TEXT("Toolset '%s' already registered."),
+				*ToolsetHandler->GetToolsetName());
+			return false;
+		}
+		ToolsetHandlers.Add(ToolsetHandler->GetToolsetName(), ToolsetHandler);
+		return true;
+	}
+
+	bool FToolsetRegistry::UnregisterToolset(TSharedPtr<IToolsetHandler> ToolsetHandler)
+	{
+		if (!ToolsetHandler.IsValid())
+		{
+			UE_LOG(LogAIAssistant, Warning, TEXT("Toolset must not be null"));
+			return false;
+		}
+		if (ToolsetHandlers.Remove(ToolsetHandler->GetToolsetName()) == 0)
+		{
+			UE_LOG(
+				LogAIAssistant, Warning, TEXT("Toolset '%s' not registered"),
+				*ToolsetHandler->GetToolsetName());
+			return false;
+		}
+		return true;
+	}
+
+	TSharedPtr<IToolsetHandler>* FToolsetRegistry::Find(
+		const FString& ToolsetName, bool LogErrors, FString* ErrorMessage)
+	{
+		TSharedPtr<IToolsetHandler>* ToolsetHandler = ToolsetHandlers.Find(ToolsetName);
+		if (ToolsetHandler == nullptr)
+		{
+			if (LogErrors || ErrorMessage)
+			{
+				FString Message = FString::Printf(TEXT("Toolset '%s' not found"), *ToolsetName);
+				if (LogErrors)
+				{
+					UE_LOG(LogAIAssistant, Error, TEXT("%s"), *Message);
+				}
+				if (ErrorMessage != nullptr)
+				{
+					*ErrorMessage = Message;
+				}
+			}
+			return nullptr;
+		}
+		return ToolsetHandler;
+	}
+
+	TFuture<TValueOrError<FJsonObject, FString>> FToolsetRegistry::ExecuteTool(
+		const FString& ToolsetAndToolName, const FJsonObject& JsonInput)
+	{
+		// Split the tool name to get the toolset name and the actual tool name.
+		TValueOrError<FToolDescriptor, FString> ToolDescriptor =
+			FToolDescriptor::FromString(ToolsetAndToolName, true);
+		if (ToolDescriptor.HasError())
+		{
+			return MakeFulfilledPromise<TValueOrError<FJsonObject, FString>>(
+				MakeError(ToolDescriptor.StealError())).GetFuture();
+		}
+
+		// Execute the tool.
+		return ExecuteTool(ToolDescriptor.StealValue(), JsonInput);
+	}
+
+	TFuture<TValueOrError<FJsonObject, FString>> FToolsetRegistry::ExecuteTool(
+		const FToolDescriptor& ToolDescriptor, const FJsonObject& JsonInput)
+	{
+		FString ErrorMessage;
+		TSharedPtr<IToolsetHandler>* ToolsetHandler =
+			Find(ToolDescriptor.ToolsetName, true, &ErrorMessage);
+		if (!ToolsetHandler)
+		{
+			return MakeFulfilledPromise<TValueOrError<FJsonObject, FString>>(
+				MakeError(ErrorMessage)).GetFuture();
+		}
+
+		// Execute the tool.
+		return (*ToolsetHandler)->ExecuteTool(ToolDescriptor.ToolName, JsonInput);
+	}
+}
diff --git a/Engine/Plugins/Experimental/AIAssistant/Source/AIAssistant/Private/AIAssistantToolsetRegistry.h b/Engine/Plugins/Experimental/AIAssistant/Source/AIAssistant/Private/AIAssistantToolsetRegistry.h
new file mode 100644
--- /dev/null
+++ b/Engine/Plugins/Experimental/AIAssistant/Source/AIAssistant/Private/AIAssistantToolsetRegistry.h
@@ -0,0 +1,58 @@
+// Copyright Epic Games, Inc. All Rights Reserved.
+
+#pragma once
+
+#include "Containers/Map.h"
+#include "Containers/UnrealString.h"
+#include "Dom/JsonObject.h"
+#include "Misc/Variant.h"
+#include "Templates/SharedPointer.h"
+
+#include "AIAssistantToolsetHandler.h"
+
+namespace UE::AIAssistant {
+
+	// Descriptor for a tool within a toolset.
+	struct FToolDescriptor
+	{
+		FString ToolsetName;
+		FString ToolName;
+
+		static TValueOrError<FToolDescriptor, FString> FromString(
+			const FString& ToolsetAndToolName, bool LogErrors = false);
+	};
+
+	// Registry for AI Assistant toolsets. Allows registering, unregistering, finding, and executing
+	// tools from toolsets.
+	class FToolsetRegistry {
+	public:
+		// Register a toolset handler. Returns true if registration was successful.
+		bool RegisterToolset(TSharedPtr<IToolsetHandler> ToolsetHandler);
+
+		// Unregister a toolset handler. Returns true if unregistration was successful.
+		bool UnregisterToolset(TSharedPtr<IToolsetHandler> ToolsetHandler);
+
+		// Find a registered toolset handler by name, and optionally log errors.
+		TSharedPtr<IToolsetHandler>* Find(
+			const FString& ToolsetName, bool LogErrors = false, FString* ErrorMessasge = nullptr);
+
+		// Execute a tool from a registered toolset handler.
+		// 
+		// If there is an error in attempting to execute a tool, such as a tool not existing,
+		// the future will contain an error string.
+		TFuture<TValueOrError<FJsonObject, FString>> ExecuteTool(
+			const FString& ToolName, const FJsonObject& JsonInput);
+
+		// Execute a tool from a registered toolset handler.
+		// 
+		// If there is an error in attempting to execute a tool, such as a tool not existing,
+		// the future will contain an error string.
+		TFuture<TValueOrError<FJsonObject, FString>> ExecuteTool(
+			const FToolDescriptor& ToolDescriptor, const FJsonObject& JsonInput);
+
+	private:
+		// The set of registered toolset handlers.
+		TMap<FString, TSharedPtr<IToolsetHandler>> ToolsetHandlers;
+	};
+
+}
diff --git a/Engine/Plugins/Experimental/AIAssistant/Source/AIAssistant/Private/Tests/AIAssistantToolsetRegistryTest.cpp b/Engine/Plugins/Experimental/AIAssistant/Source/AIAssistant/Private/Tests/AIAssistantToolsetRegistryTest.cpp
new file mode 100644
--- /dev/null
+++ b/Engine/Plugins/Experimental/AIAssistant/Source/AIAssistant/Private/Tests/AIAssistantToolsetRegistryTest.cpp
@@ -0,0 +1,468 @@
+// Copyright Epic Games, Inc. All Rights Reserved.
+
+#include "AIAssistantToolsetRegistry.h"
+
+#include "Dom/JsonObject.h"
+#include "Misc/AutomationTest.h"
+#include "UObject/StrongObjectPtr.h"
+
+#include "AIAssistantTestFlags.h"
+#include "AIAssistantToolsetHandler.h"
+
+using namespace UE::AIAssistant;
+
+IMPLEMENT_SIMPLE_AUTOMATION_TEST(
+	FAIAssistantToolsetDescriptorFromStringValidName,
+	"AI.Assistant.ToolsetDescriptor.FromString.ValidName",
+	AIAssistantTest::Flags);
+
+bool FAIAssistantToolsetDescriptorFromStringValidName::RunTest(const FString& UnusedParameters)
+{
+	// Test valid toolset and tool name.
+	TValueOrError<FToolDescriptor, FString> Descriptor =
+		FToolDescriptor::FromString(TEXT("ToolsetName.ToolName"));
+	(void)TestTrue(TEXT("FromString_Valid"), Descriptor.HasValue());
+	FToolDescriptor* Value = Descriptor.TryGetValue();
+	if (!Value)
+	{
+		return false;
+	}
+	(void)TestEqual(TEXT("FromString_Valid_ToolsetName"), Value->ToolsetName, TEXT("ToolsetName"));
+	(void)TestEqual(TEXT("FromString_Valid_ToolName"), Value->ToolName, TEXT("ToolName"));
+
+	return true;
+}
+
+IMPLEMENT_SIMPLE_AUTOMATION_TEST(
+	FAIAssistantToolsetDescriptorFromStringInvalidToolsetName,
+	"AI.Assistant.ToolsetDescriptor.FromString.InvalidToolsetName",
+	AIAssistantTest::Flags);
+
+bool FAIAssistantToolsetDescriptorFromStringInvalidToolsetName::RunTest(
+	const FString& UnusedParameters)
+{
+	// Test invalid toolset and tool name (missing dot).
+	TValueOrError<FToolDescriptor, FString> Descriptor =
+		FToolDescriptor::FromString(TEXT("InvalidToolNameWithoutDot"));
+	(void)TestTrue(TEXT("FromString_Invalid_MissingDot"), Descriptor.HasError());
+
+	return true;
+}
+
+IMPLEMENT_SIMPLE_AUTOMATION_TEST(
+	FAIAssistantToolsetDescriptorFromStringMissingToolsetName,
+	"AI.Assistant.ToolsetDescriptor.FromString.MissingToolsetName",
+	AIAssistantTest::Flags);
+
+bool FAIAssistantToolsetDescriptorFromStringMissingToolsetName::RunTest(
+	const FString& UnusedParameters)
+{
+	// Test missing toolset name.
+	TValueOrError<FToolDescriptor, FString> Descriptor =
+		FToolDescriptor::FromString(TEXT(".ToolNameWithoutToolset"));
+	(void)TestTrue(TEXT("FromString_Invalid_MissingToolsetName"), Descriptor.HasError());
+
+	return true;
+}
+
+IMPLEMENT_SIMPLE_AUTOMATION_TEST(
+	FAIAssistantToolsetDescriptorFromStringMissingToolName,
+	"AI.Assistant.ToolsetDescriptor.FromString.MissingToolName",
+	AIAssistantTest::Flags);
+
+bool FAIAssistantToolsetDescriptorFromStringMissingToolName::RunTest(
+	const FString& UnusedParameters)
+{
+	// Test missing tool name.
+	TValueOrError<FToolDescriptor, FString> Descriptor =
+		FToolDescriptor::FromString(TEXT("ToolsetWithoutToolname."));
+	(void)TestTrue(TEXT("FromString_Invalid_MissingToolName"), Descriptor.HasError());
+
+	return true;
+}
+
+IMPLEMENT_SIMPLE_AUTOMATION_TEST(
+	FAIAssistantToolsetDescriptorFromStringMultipleDots,
+	"AI.Assistant.ToolsetDescriptor.FromString.MultipleDots",
+	AIAssistantTest::Flags);
+
+bool FAIAssistantToolsetDescriptorFromStringMultipleDots::RunTest(const FString& UnusedParameters)
+{
+	// Test valid toolset and tool name with multiple dots.
+	TValueOrError<FToolDescriptor, FString> Descriptor =
+		FToolDescriptor::FromString(TEXT("Complex.Toolset.Name.ToolName"));
+	(void)TestTrue(TEXT("FromString_Valid_MultipleDots"), Descriptor.HasValue());
+	FToolDescriptor* Value = Descriptor.TryGetValue();
+	if (!Value)
+	{
+		return false;
+	}
+	(void)TestEqual(
+		TEXT("FromString_Valid_MultipleDots_ToolsetName"),
+		Descriptor.TryGetValue()->ToolsetName,
+		TEXT("Complex.Toolset.Name"));
+	(void)TestEqual(
+		TEXT("FromString_Valid_MultipleDots_ToolName"),
+		Descriptor.TryGetValue()->ToolName,
+		TEXT("ToolName"));
+
+	return true;
+}
+
+class FakeToolset : public IToolsetHandler
+{
+public:
+	TFuture<TValueOrError<FJsonObject, FString>> ExecuteTool(
+		const FString& ToolName, const FJsonObject& JsonInput) override
+	{
+		if (ToolName != TEXT("SomeTool"))
+		{
+			return MakeFulfilledPromise<TValueOrError<FJsonObject, FString>>(
+				MakeError(FString(TEXT("Tool not found")))).GetFuture();
+		}
+
+		FJsonObject Result;
+		Result.SetStringField("result", "success");
+		return MakeFulfilledPromise<TValueOrError<FJsonObject, FString>>(
+			MakeValue(Result)).GetFuture();
+	}
+
+	FJsonObject GetToolsetJsonSchema() const override
+	{
+		return FJsonObject();
+	}
+
+	FString GetToolsetName() const override
+	{
+		return TEXT("Fake");
+	}
+
+	FString GetToolsetVersion() const override
+	{
+		return TEXT("1.0");
+	}
+};
+
+class AnotherFakeToolset : public FakeToolset
+{
+public:
+	FString GetToolsetName() const override
+	{
+		return TEXT("AnotherFake");
+	}
+};
+
+IMPLEMENT_SIMPLE_AUTOMATION_TEST(
+	FAIAssistantToolsetRegistryToolsetHandlerRegistrationSuccess,
+	"AI.Assistant.ToolsetRegistry.Registration.Success",
+	AIAssistantTest::Flags);
+
+bool FAIAssistantToolsetRegistryToolsetHandlerRegistrationSuccess::RunTest(
+	const FString& UnusedParameters)
+{
+	TUniquePtr<FToolsetRegistry> Registry = MakeUnique<FToolsetRegistry>();
+
+	// Test whether a valid toolset handler is accepted.
+	TSharedPtr<IToolsetHandler> ValidToolsetHandler = MakeShared<FakeToolset>();
+	(void)TestTrue(
+		TEXT("RegisterToolsetHandler_Valid"),
+		Registry->RegisterToolset(ValidToolsetHandler));
+
+	return true;
+}
+
+IMPLEMENT_SIMPLE_AUTOMATION_TEST(
+	FAIAssistantToolsetRegistryToolsetHandlerRegistrationDoubleRegistration,
+	"AI.Assistant.ToolsetRegistry.Registration.DoubleRegistration",
+	AIAssistantTest::Flags);
+
+bool FAIAssistantToolsetRegistryToolsetHandlerRegistrationDoubleRegistration::RunTest(
+	const FString& UnusedParameters)
+{
+	TUniquePtr<FToolsetRegistry> Registry = MakeUnique<FToolsetRegistry>();
+	TSharedPtr<IToolsetHandler> ValidToolsetHandler = MakeShared<FakeToolset>();
+	Registry->RegisterToolset(ValidToolsetHandler);
+
+	// Test whether duplicate registration is rejected.
+	AddExpectedError(TEXT("Toolset 'Fake' already registered"));
+	(void)TestFalse(
+		TEXT("RegisterToolsetHandler_Duplicate"),
+		Registry->RegisterToolset(ValidToolsetHandler));
+
+	return true;
+}
+
+IMPLEMENT_SIMPLE_AUTOMATION_TEST(
+	FAIAssistantToolsetRegistryToolsetHandlerUnregistrationSuccess,
+	"AI.Assistant.ToolsetRegistry.Unregistration.Success",
+	AIAssistantTest::Flags);
+
+bool FAIAssistantToolsetRegistryToolsetHandlerUnregistrationSuccess::RunTest(
+	const FString& UnusedParameters)
+{
+	TUniquePtr<FToolsetRegistry> Registry = MakeUnique<FToolsetRegistry>();
+	TSharedPtr<IToolsetHandler> ValidToolsetHandler = MakeShared<FakeToolset>();
+	Registry->RegisterToolset(ValidToolsetHandler);
+
+	// Test whether unregistration of a registered toolset handler is accepted.
+	(void)TestTrue(
+		TEXT("UnregisterToolsetHandler_Valid"),
+		Registry->UnregisterToolset(ValidToolsetHandler));
+
+	return true;
+}
+
+IMPLEMENT_SIMPLE_AUTOMATION_TEST(
+	FAIAssistantToolsetRegistryToolsetHandlerUnregistrationUnregisteredToolset,
+	"AI.Assistant.ToolsetRegistry.Unregistration.UnregisteredToolset",
+	AIAssistantTest::Flags);
+
+bool FAIAssistantToolsetRegistryToolsetHandlerUnregistrationUnregisteredToolset::RunTest(
+	const FString& UnusedParameters)
+{
+	TUniquePtr<FToolsetRegistry> Registry = MakeUnique<FToolsetRegistry>();
+	Registry->RegisterToolset(MakeShared<FakeToolset>());
+
+	// Test whether unregistration of an non-registered toolset handler is rejected.
+	TSharedPtr<IToolsetHandler> UnregisteredToolsetHandler = MakeShared<AnotherFakeToolset>();
+	AddExpectedError(TEXT("Toolset 'AnotherFake' not registered"));
+	(void)TestFalse(
+		TEXT("UnregisterToolsetHandler_NotRegistered"),
+		Registry->UnregisterToolset(UnregisteredToolsetHandler));
+
+	return true;
+}
+
+IMPLEMENT_SIMPLE_AUTOMATION_TEST(
+	FAIAssistantToolsetRegistryToolsetHandlerUnregistrationDoubleUnregistration,
+	"AI.Assistant.ToolsetRegistry.Unregistration.DoubleUnregistration",
+	AIAssistantTest::Flags);
+
+bool FAIAssistantToolsetRegistryToolsetHandlerUnregistrationDoubleUnregistration::RunTest(
+	const FString& UnusedParameters)
+{
+	TUniquePtr<FToolsetRegistry> Registry = MakeUnique<FToolsetRegistry>();
+	TSharedPtr<IToolsetHandler> ValidToolsetHandler = MakeShared<FakeToolset>();
+	Registry->RegisterToolset(ValidToolsetHandler);
+	Registry->UnregisterToolset(ValidToolsetHandler);
+
+	// Test whether unregistration of an already unregistered toolset handler is rejected.
+	AddExpectedError(TEXT("Toolset 'Fake' not registered"));
+	(void)TestFalse(
+		TEXT("UnregisterToolsetHandler_AlreadyUnregistered"),
+		Registry->UnregisterToolset(ValidToolsetHandler));
+
+	return true;
+}
+
+IMPLEMENT_SIMPLE_AUTOMATION_TEST(
+	FAIAssistantToolsetRegistryToolsetHandlerFindSuccess,
+	"AI.Assistant.ToolsetRegistry.Find.Success",
+	AIAssistantTest::Flags);
+
+bool FAIAssistantToolsetRegistryToolsetHandlerFindSuccess::RunTest(const FString& UnusedParameters)
+{
+	TUniquePtr<FToolsetRegistry> Registry = MakeUnique<FToolsetRegistry>();
+	TSharedPtr<IToolsetHandler> Fake = MakeShared<FakeToolset>();
+	TSharedPtr<IToolsetHandler> AnotherFake = MakeShared<AnotherFakeToolset>();
+	Registry->RegisterToolset(Fake);
+	Registry->RegisterToolset(AnotherFake);
+
+	// Test finding registered toolset handlers.
+	TSharedPtr<IToolsetHandler>* FoundToolsetHandler = Registry->Find(TEXT("Fake"));
+	(void)TestNotNull(TEXT("FindToolsetHandler_Registered_NotNull"), FoundToolsetHandler);
+	if (!FoundToolsetHandler)
+	{
+		return false;
+	}
+	(void)TestEqual(TEXT("FindToolsetHandler_Registered_Correct"), *FoundToolsetHandler, Fake);
+
+	FoundToolsetHandler = Registry->Find(TEXT("AnotherFake"));
+	(void)TestNotNull(TEXT("FindToolsetHandler_AnotherRegistered_NotNull"), FoundToolsetHandler);
+	if (!FoundToolsetHandler)
+	{
+		return false;
+	}
+	(void)TestEqual(
+		TEXT("FindToolsetHandler_AnotherRegistered_Correct"), *FoundToolsetHandler, AnotherFake);
+
+	return true;
+}
+
+IMPLEMENT_SIMPLE_AUTOMATION_TEST(
+	FAIAssistantToolsetRegistryToolsetHandlerFindNonexistentToolsetLogError,
+	"AI.Assistant.ToolsetRegistry.Find.NonexistentToolset.LogError",
+	AIAssistantTest::Flags);
+
+bool FAIAssistantToolsetRegistryToolsetHandlerFindNonexistentToolsetLogError::RunTest(
+	const FString& UnusedParameters)
+{
+	TUniquePtr<FToolsetRegistry> Registry = MakeUnique<FToolsetRegistry>();
+	TSharedPtr<IToolsetHandler> Fake = MakeShared<FakeToolset>();
+	Registry->RegisterToolset(Fake);
+	
+	// Test finding an unregistered toolset handler.
+	AddExpectedError(TEXT("Toolset 'NonExistentToolset' not found"));
+	TSharedPtr<IToolsetHandler>* FoundToolsetHandler =
+		Registry->Find(TEXT("NonExistentToolset"), true);
+	(void)TestNull(TEXT("FindToolsetHandler_Unregistered"), FoundToolsetHandler);
+
+	return true;
+}
+
+IMPLEMENT_SIMPLE_AUTOMATION_TEST(
+	FAIAssistantToolsetRegistryToolsetHandlerFindNonexistentToolsetReturnErrorMessage,
+	"AI.Assistant.ToolsetRegistry.Find.NonexistentToolset.ReturnErrorMessage",
+	AIAssistantTest::Flags);
+
+bool FAIAssistantToolsetRegistryToolsetHandlerFindNonexistentToolsetReturnErrorMessage::RunTest(
+	const FString& UnusedParameters)
+{
+	TUniquePtr<FToolsetRegistry> Registry = MakeUnique<FToolsetRegistry>();
+	TSharedPtr<IToolsetHandler> Fake = MakeShared<FakeToolset>();
+	Registry->RegisterToolset(Fake);
+
+	// Test finding an unregistered toolset handler with error message.
+	FString ErrorMessage;
+	TSharedPtr<IToolsetHandler>* FoundToolsetHandler =
+		Registry->Find(TEXT("NonExistentToolset"), false, &ErrorMessage);
+	(void)TestNull(
+		TEXT("FindToolsetHandler_UnregisteredWithErrorMessage_Found"),
+		FoundToolsetHandler);
+	(void)TestEqual(
+		TEXT("FindToolsetHandler_UnregisteredWithErrorMessage_ErrorMessage"),
+		ErrorMessage,
+		TEXT("Toolset 'NonExistentToolset' not found"));
+
+	return true;
+}
+
+IMPLEMENT_SIMPLE_AUTOMATION_TEST(
+	FAIAssistantToolsetRegistryToolsetHandlerExecuteToolSuccess,
+	"AI.Assistant.ToolsetRegistry.ExecuteTool.Success",
+	AIAssistantTest::Flags);
+	
+bool FAIAssistantToolsetRegistryToolsetHandlerExecuteToolSuccess::RunTest(
+const FString& UnusedParameters)
+{
+	TUniquePtr<FToolsetRegistry> Registry = MakeUnique<FToolsetRegistry>();
+
+	// Register a valid toolset handler.
+	TSharedPtr<IToolsetHandler> ValidToolsetHandler = MakeShared<FakeToolset>();
+	(void)TestTrue(
+		TEXT("RegisterToolsetHandler_Valid"),
+		Registry->RegisterToolset(ValidToolsetHandler));
+
+	FJsonObject InputArguments;
+	InputArguments.SetStringField(TEXT("input"), TEXT("value"));
+
+	// Test execution of a tool from the registered toolset handler.
+	Registry->ExecuteTool(TEXT("Fake.SomeTool"), InputArguments)
+		.Next([this](TValueOrError<FJsonObject, FString> Result)
+		{
+			(void)TestTrue(TEXT("ExecuteTool_SuccessValue"), Result.HasValue());
+			(void)TestEqual(
+				TEXT("ExecuteTool_ResultJson"),
+				*Result.StealValue().GetStringField(TEXT("result")),
+				TEXT("success"));
+		})
+		.Wait();
+
+	return true;
+}
+
+IMPLEMENT_SIMPLE_AUTOMATION_TEST(
+	FAIAssistantToolsetRegistryToolsetHandlerExecuteToolNonexistantTool,
+	"AI.Assistant.ToolsetRegistry.ExecuteTool.NonexistantTool",
+	AIAssistantTest::Flags);
+
+bool FAIAssistantToolsetRegistryToolsetHandlerExecuteToolNonexistantTool::RunTest(
+	const FString& UnusedParameters)
+{
+	TUniquePtr<FToolsetRegistry> Registry = MakeUnique<FToolsetRegistry>();
+
+	// Register a valid toolset handler.
+	TSharedPtr<IToolsetHandler> ValidToolsetHandler = MakeShared<FakeToolset>();
+	(void)TestTrue(
+		TEXT("RegisterToolsetHandler_Valid"),
+		Registry->RegisterToolset(ValidToolsetHandler));
+
+	// Test execution of an invalid tool from the registered toolset handler.
+	Registry->ExecuteTool(TEXT("Fake.SomeNonexistantTool"), FJsonObject())
+		.Next([this](TValueOrError<FJsonObject, FString> Result)
+		{
+			(void)TestTrue(TEXT("NonexistantToolset_HasError"), Result.HasError());
+			(void)TestEqual(
+				TEXT("ExecuteTool_InvalidTool_Message"),
+				*Result.GetError(),
+				TEXT("Tool not found"));
+		})
+		.Wait();
+
+	return true;
+}
+
+IMPLEMENT_SIMPLE_AUTOMATION_TEST(
+	FAIAssistantToolsetRegistryToolsetHandlerExecuteToolNonexistantToolset,
+	"AI.Assistant.ToolsetRegistry.ExecuteTool.NonexistantToolset",
+	AIAssistantTest::Flags);
+
+bool FAIAssistantToolsetRegistryToolsetHandlerExecuteToolNonexistantToolset::RunTest(
+	const FString& UnusedParameters)
+{
+	TUniquePtr<FToolsetRegistry> Registry = MakeUnique<FToolsetRegistry>();
+
+	// Register a valid toolset handler.
+	TSharedPtr<IToolsetHandler> ValidToolsetHandler = MakeShared<FakeToolset>();
+	(void)TestTrue(
+		TEXT("RegisterToolsetHandler_Valid"),
+		Registry->RegisterToolset(ValidToolsetHandler));
+
+	// Test execution of an invalid toolset.
+	AddExpectedError(TEXT("Toolset 'InvalidTool' not found"));
+	Registry->ExecuteTool(TEXT("InvalidTool.SolveHaltingProblem"), FJsonObject())
+		.Next([this](TValueOrError<FJsonObject, FString> Result)
+			{
+				(void)TestTrue(TEXT("ExecuteTool_InvalidToolset_HasError"), Result.HasError());
+				(void)TestEqual(
+					TEXT("ExecuteTool_InvalidToolset_Message"),
+					*Result.GetError(),
+					TEXT("Toolset 'InvalidTool' not found"));
+			})
+		.Wait();
+
+	return true;
+}
+
+IMPLEMENT_SIMPLE_AUTOMATION_TEST(
+	FAIAssistantToolsetRegistryToolsetHandlerExecuteToolBadToolnameFormat,
+	"AI.Assistant.ToolsetRegistry.ExecuteTool.BadToolnameFormat",
+	AIAssistantTest::Flags);
+
+bool FAIAssistantToolsetRegistryToolsetHandlerExecuteToolBadToolnameFormat::RunTest(
+	const FString& UnusedParameters)
+{
+	TUniquePtr<FToolsetRegistry> Registry = MakeUnique<FToolsetRegistry>();
+
+	// Register a valid toolset handler.
+	TSharedPtr<IToolsetHandler> ValidToolsetHandler = MakeShared<FakeToolset>();
+	(void)TestTrue(
+		TEXT("RegisterToolsetHandler_Valid"),
+		Registry->RegisterToolset(ValidToolsetHandler));
+
+	// Test execution of an incorrectly formatted tool name.
+	AddExpectedError(TEXT("Invalid tool name format. Expected format: ToolsetName.ToolName"));
+	Registry->ExecuteTool(TEXT("ToolCallThatDoesNotHaveADot"), FJsonObject())
+		.Next([this](TValueOrError<FJsonObject, FString> Result)
+			{
+				(void)TestTrue(TEXT("ExecuteTool_InvalidToolName_HasError"), Result.HasError());
+				(void)TestEqual(
+					TEXT("ExecuteTool_InvalidToolName_ResultJson"),
+					*Result.GetError(),
+					TEXT("Invalid tool name format. Expected format: ToolsetName.ToolName"));
+			})
+		.Wait();
+
+
+	return true;
+}
