[Backout] - CL47483013
#rnx
Original CL Desc
-----------------------------------------------------------------
Added IToolsetHandler and a Toolset registry, along with functions to register, and unregister toolsets, and execute tool calls.

diff --git a/Engine/Plugins/Experimental/AIAssistant/Source/AIAssistant/Private/AIAssistantSubsystem.h b/Engine/Plugins/Experimental/AIAssistant/Source/AIAssistant/Private/AIAssistantSubsystem.h
--- a/Engine/Plugins/Experimental/AIAssistant/Source/AIAssistant/Private/AIAssistantSubsystem.h
+++ b/Engine/Plugins/Experimental/AIAssistant/Source/AIAssistant/Private/AIAssistantSubsystem.h
@@ -3,9 +3,7 @@
 
 #pragma once
 
-#include "Containers/UnrealString.h"
 #include "EditorSubsystem.h"
-#include "AIAssistantToolsetRegistry.h"
 #include "Templates/SharedPointer.h"
 
 #include "AIAssistantSubsystem.generated.h"
@@ -45,7 +43,4 @@ public:
 	static FAIAssistantModule& GetAIAssistantModule();
 	// Get the assistant web browser.
 	static TSharedPtr<SAIAssistantWebBrowser> GetAIAssistantWebBrowserWidget();
-
-	// The set of registered toolset handlers.
-	UE::AIAssistant::FToolsetRegistry ToolsetRegistry;
 };
diff --git a/Engine/Plugins/Experimental/AIAssistant/Source/AIAssistant/Private/AIAssistantToolsetHandler.h b/Engine/Plugins/Experimental/AIAssistant/Source/AIAssistant/Private/AIAssistantToolsetHandler.h
deleted file mode 100644
diff --git a/Engine/Plugins/Experimental/AIAssistant/Source/AIAssistant/Private/AIAssistantToolsetRegistry.cpp b/Engine/Plugins/Experimental/AIAssistant/Source/AIAssistant/Private/AIAssistantToolsetRegistry.cpp
deleted file mode 100644
diff --git a/Engine/Plugins/Experimental/AIAssistant/Source/AIAssistant/Private/AIAssistantToolsetRegistry.h b/Engine/Plugins/Experimental/AIAssistant/Source/AIAssistant/Private/AIAssistantToolsetRegistry.h
deleted file mode 100644
diff --git a/Engine/Plugins/Experimental/AIAssistant/Source/AIAssistant/Private/Tests/AIAssistantToolsetRegistryTest.cpp b/Engine/Plugins/Experimental/AIAssistant/Source/AIAssistant/Private/Tests/AIAssistantToolsetRegistryTest.cpp
deleted file mode 100644
