[AI Assistant] Refactor Slate querying code and add tests

#rnx

diff --git a/Engine/Plugins/Experimental/AIAssistant/Source/AIAssistant/Private/AIAssistantSlateQuerier.cpp b/Engine/Plugins/Experimental/AIAssistant/Source/AIAssistant/Private/AIAssistantSlateQuerier.cpp
--- a/Engine/Plugins/Experimental/AIAssistant/Source/AIAssistant/Private/AIAssistantSlateQuerier.cpp
+++ b/Engine/Plugins/Experimental/AIAssistant/Source/AIAssistant/Private/AIAssistantSlateQuerier.cpp
@@ -1,18 +1,17 @@
 // Copyright Epic Games, Inc. All Rights Reserved.
-	
 
 #include "AIAssistantSlateQuerier.h"
 
 #include "EditorModes.h"
 #include "Containers/Set.h"
 #include "Containers/ArrayView.h"
-#include "Internationalization/Regex.h"
 #include "LevelEditorSubsystem.h"
 #include "LevelEditorViewport.h"
 #include "Settings/LevelEditorViewportSettings.h"
 #include "EditorModeManager.h"
 #include "EngineUtils.h"
 #include "IDetailsView.h"
+#include "Layout/WidgetPath.h"
 #include "Interfaces/IMainFrameModule.h"
 #include "Internationalization/Internationalization.h"
 #include "Internationalization/Culture.h"
@@ -21,8 +20,6 @@
 #include "Subsystems/AssetEditorSubsystem.h"
 #include "Subsystems/EditorActorSubsystem.h"
 #include "Toolkits/BaseToolkit.h"
-#include "Widgets/Text/STextBlock.h"
-#include "Widgets/Text/SRichTextBlock.h"
 #include "Widgets/Input/SMultiLineEditableTextBox.h"
 #include "Widgets/Input/SSearchBox.h"
 #include "Widgets/Docking/SDockTab.h"
@@ -32,816 +29,342 @@
 
 #include "AIAssistant.h"
 #include "AIAssistantLog.h"
+#include "AIAssistantSlateQuerierUtils.h"
+#include "AIAssistantSlateQuerierUtilsTypeSearch.h"
 #include "AIAssistantSubsystem.h"
 #include "AIAssistantWebBrowser.h"
 
 
 #define LOCTEXT_NAMESPACE "FAIAssistantSlateQuerier"
 
-struct FAIAssistantSlateQueryStructuredContext : public FJsonSerializable
-{
-	FString TopWindowName;
-	FString UnrealEditorVersion;
-	FString UnrealEditorLanguage;
-	FString UnrealEditorLocale;
-	TOptional<FString> WindowName;
-	TOptional<FString> EditorName;
-	TOptional<FString> EditorMode;
-	TOptional<FString> ActiveToolName;
-	TOptional<FString> TabName;
-	TOptional<FString> ToolTipText;
-	TOptional<FString> TextUnderCursor;
-	TOptional<FString> WidgetType;
-	TOptional<FString> ActorClass;
-	TOptional<FString> ItemName;
-	TOptional<FString> ItemDescriptor;
-	TOptional<FString> CategoryName;
-	TOptional<TArray<FString>> WidgetPath;
-
-	BEGIN_JSON_SERIALIZER
-		JSON_SERIALIZE("toplevel_window_name", TopWindowName);
-		JSON_SERIALIZE("unreal_editor_version", UnrealEditorVersion);
-		JSON_SERIALIZE("unreal_editor_language", UnrealEditorLanguage);
-		JSON_SERIALIZE("unreal_editor_locale", UnrealEditorLocale);
-		JSON_SERIALIZE_OPTIONAL("window_name", WindowName);
-		JSON_SERIALIZE_OPTIONAL("editor_name", EditorName);
-		JSON_SERIALIZE_OPTIONAL("editor_mode", EditorMode);
-		JSON_SERIALIZE_OPTIONAL("active_tool", ActiveToolName);
-		JSON_SERIALIZE_OPTIONAL("tab_name", TabName);
-		JSON_SERIALIZE_OPTIONAL("tooltip_text", ToolTipText);
-		JSON_SERIALIZE_OPTIONAL("text_under_cursor", TextUnderCursor);
-		JSON_SERIALIZE_OPTIONAL("widget_type", WidgetType);
-		JSON_SERIALIZE_OPTIONAL("actor_class", ActorClass);
-		JSON_SERIALIZE_OPTIONAL("item_name", ItemName);
-		JSON_SERIALIZE_OPTIONAL("item_descriptor", ItemDescriptor);
-		JSON_SERIALIZE_OPTIONAL("category_name", CategoryName);
-		JSON_SERIALIZE_OPTIONAL_ARRAY("widget_path", WidgetPath);
-	END_JSON_SERIALIZER
-};
-
+using namespace UE::AIAssistant::SlateQuerier;
 
 //
-// FAIAssistantSlateQueryContext
-//
-// Used for accumulating UI context information, and ultimately building strings for the query.
+// UE::AIAssistant::SlateQuerier
 //
-
-
-struct FAIAssistantSlateQueryContext
+FWidgetPath UE::AIAssistant::SlateQuerier::GetWidgetPathUnderCursor()
 {
-	TSharedPtr<SWidget> LastPickedWidget;
-	FText CurrentToolTipText = FText();
-	FText FormattedWindowName = FText();
-	FText FormattedTabName = FText();
-	FText TextUnderCursor = FText();
-	FText ItemName = FText();
-	FText ItemDescriptor = FText();
-	FString SelectedActorName = FString();
-	FString SelectedActorClass = FString();
-	bool bIsUIWidget = false;
-	bool bIsObject = false;
-	bool bIsInAssistantPanel = false;
-
-	FText GeneratedQuery = FText();
-	FText GeneratedQueryInstructions = FText();
-	FText LanguageInstructions = FText();
-
-	FAIAssistantSlateQueryStructuredContext StructuredContext;
-};
-
-
-//
-// Statics
-//
+	const FVector2f CursorPos = FSlateApplication::Get().GetCursorPos();
+	const TArray<TSharedRef<SWindow>>& Windows = FSlateApplication::Get().GetTopLevelWindows();
+	return FSlateApplication::Get().LocateWindowUnderMouse(CursorPos, Windows, true);
+}
 
-// Search backwards from end of WidgetPathToTest and return first widget that matches any type in WidgetTypes.
-// Return null ptr if not found.
-static TSharedPtr<SWidget> FindClosestWidgetPtrOfTypes(const FWidgetPath& WidgetPathToTest, const TSet<FName> WidgetTypes)
+void UE::AIAssistant::SlateQuerier::QueryAIAssistantAboutSlateWidget(const FWidgetPath& InWidgetPath)
 {
-	TSharedPtr<SWidget> ClosestWidget;
-	for (int32 WidgetIndex = WidgetPathToTest.Widgets.Num() - 1; WidgetIndex >= 0; --WidgetIndex)
+	if (!InWidgetPath.IsValid())
 	{
-		const FArrangedWidget* ArrangedWidget = &WidgetPathToTest.Widgets[WidgetIndex];
-		const TSharedRef<SWidget>& ThisWidget = ArrangedWidget->Widget;
-		if (WidgetTypes.Contains(ThisWidget->GetType()))
-		{
-			ClosestWidget = ThisWidget.ToSharedPtr();
-			break;
-		}
+		UE_LOG(LogAIAssistant, Warning, TEXT("No valid widget path found to generate widget query."));
+		return;
 	}
-	return ClosestWidget;
-}
-
-// Search backwards from end of WidgetPathToTest and return first widget that matches WidgetType.  Return null
-// ptr if not found.
-static TSharedPtr<SWidget> FindClosestWidgetPtrOfType(const FWidgetPath& WidgetPathToTest, const FName& WidgetType)
-{
-	const TSet<FName> WidgetTypes = TSet<FName>({WidgetType});
-	return FindClosestWidgetPtrOfTypes(WidgetPathToTest, WidgetTypes);
-}
 
-// Search backwards from end of WidgetPathToTest and return first widget that matches a type in WidgetTypes,
-// testing each type in order.  Return null ptr if not found.
-static TSharedPtr<SWidget> FindClosestWidgetPtrOfSortedTypes(
-	const FWidgetPath& WidgetPathToTest, const TConstArrayView<FName>& WidgetTypes)
-{
-	for (FName WidgetType : WidgetTypes)
+	FSlateQuerier SlateQuerier = FSlateQuerier(InWidgetPath);
+	if (!SlateQuerier.GenerateAndSendQuery())
 	{
-		TSharedPtr<SWidget> ClosestWidget = FindClosestWidgetPtrOfType(WidgetPathToTest, WidgetType);
-		if (ClosestWidget.IsValid())
-		{
-			return ClosestWidget;
-		}
+		UE_LOG(LogAIAssistant, Warning, TEXT("Could not generate query for widget."));
 	}
-
-	return nullptr;
 }
 
-// Search forward from start of WidgetPathToTest and return first widget that matches WidgetType.  Return null
-// ptr if not found.
-static TSharedPtr<SWidget> FindFirstWidgetPtrOfType(const FWidgetPath& WidgetPathToTest, const FName& WidgetType)
+FSlateQuerier::FSlateQuerier(const FWidgetPath& InWidgetPath) 
+	: OriginalWidgetPath(InWidgetPath), WidgetPath(InWidgetPath), bIsUIWidget(false), bIsObject(false),
+	  bIsInAssistantPanel(false), GeneratedQuery(), StructuredContext()
+{}
+
+bool FSlateQuerier::GenerateAndSendQuery()
 {
-	TSharedPtr<SWidget> FirstWidget;
-	for (int32 WidgetIndex = 0; WidgetIndex < WidgetPathToTest.Widgets.Num(); WidgetIndex++)
+	IdentifyObjectOfQuery();
+	GenerateQuery();
+	if (IsValid())
 	{
-		const FArrangedWidget* ArrangedWidget = &WidgetPathToTest.Widgets[WidgetIndex];
-		const TSharedRef<SWidget>& ThisWidget = ArrangedWidget->Widget;
-		if (ThisWidget->GetType() == WidgetType)
-		{
-			FirstWidget = ThisWidget.ToSharedPtr();
-			break;
-		}
+		GenerateStructuredContext();
+		SendQuery();
 	}
-	return FirstWidget;
+	return IsValid();
 }
 
-// Search children of WidgetToTest for any widgets that match WidgetType; return all found widgets in OutWidgets.
-static void FindChildWidgetsOfType(TArray<TSharedRef<SWidget>>& OutWidgets, TSharedPtr<SWidget> WidgetToTest, const FName& WidgetType)
+void FSlateQuerier::IdentifyObjectOfQuery()
 {
-	if (WidgetToTest.IsValid())
+	// For menu items, we want to generate the context based on the root widget that spawned the menu,
+	// not on the actual menu item itself.
+	if (Utility::FindClosestMenuItem(OriginalWidgetPath))
 	{
-		for (int32 ChildIdx = 0; ChildIdx < WidgetToTest->GetChildren()->Num(); ChildIdx++)
+		TSharedPtr<SWidget> RootWidget = FSlateApplication::Get().GetMenuHostWidget();
+		if (RootWidget.IsValid())
 		{
-			TSharedRef<SWidget> ThisWidget = WidgetToTest->GetChildren()->GetChildAt(ChildIdx);
-			if (ThisWidget->GetType() == WidgetType)
-			{
-				OutWidgets.Add(ThisWidget);
-			}
-			FindChildWidgetsOfType(OutWidgets, ThisWidget.ToSharedPtr(), WidgetType);
+			FSlateApplication::Get().GeneratePathToWidgetChecked(RootWidget.ToSharedRef(), WidgetPath);
 		}
 	}
-}
 
-// Search children of WidgetToTest and return first widget that matches WidgetType.  Return null ptr if not found.
-static TSharedPtr<SWidget> FindChildWidgetOfType(const TSharedPtr<SWidget> WidgetToTest, const FName& WidgetType)
-{
-	if (WidgetToTest.IsValid())
+	if (!WidgetPath.IsValid())
 	{
-		for (int32 ChildIdx = 0; ChildIdx < WidgetToTest->GetChildren()->Num(); ChildIdx++)
-		{
-			TSharedRef<SWidget> ThisWidget = WidgetToTest->GetChildren()->GetChildAt(ChildIdx);
-			if (ThisWidget->GetType() == WidgetType)
-			{
-				return ThisWidget.ToSharedPtr();
-			}
-			else
-			{
-				TSharedPtr<SWidget> ChildWidget = FindChildWidgetOfType(ThisWidget.ToSharedPtr(), WidgetType);
-				if (ChildWidget.IsValid() && ChildWidget->GetType() == WidgetType)
-				{
-					return ChildWidget;
-				}
-			}
-		}
+		return;
 	}
-	return TSharedPtr<SWidget>();
+
+	// Is there a current tool-tip?
+	CurrentToolTipText = Utility::FindCurrentToolTipText();
+
+	// Is there an identifiable asset editor, mode, or window?
+	FindEditorName();
+
+	// Is there an identifiable tab?
+	FindTabName();
+
+	// Is there text under the cursor??
+	TextUnderCursor = Utility::FindTextUnderCursor(WidgetPath);
+
+	// Is there an identifiable item?
+	FindItemName();
 }
 
-// Search children of WidgetToTest and return text from first one found with alpha text.  Return empty FText
-// if not found.
-static FText FindChildWidgetWithText(const TSharedPtr<SWidget> WidgetToTest)
+void FSlateQuerier::GenerateQuery()
 {
-	// get all children and see if any are text widgets.
-	// refuse any that only contain numbers (they're probably just values)
-	const FRegexPattern AlphaPattern(TEXT("[A-Za-z]+"));
-	if (WidgetToTest.IsValid())
+	FText FormattedItemName;
+	if (!ItemName.IsEmpty())
 	{
-		if (WidgetToTest->GetType() == "STextBlock")
-		{
-			TSharedRef<STextBlock> TextBlock = StaticCastSharedRef<STextBlock>(WidgetToTest.ToSharedRef());
-			const FText WidgetText = TextBlock->GetText();
-			if (!WidgetText.IsEmpty())
-			{
-				FRegexMatcher AlphaMatcher(AlphaPattern, WidgetText.ToString());
-				if (AlphaMatcher.FindNext())
-				{
-					return WidgetText;
-				}
-			}
-		}
-		for (int32 ChildIdx = 0; ChildIdx < WidgetToTest->GetChildren()->Num(); ChildIdx++)
-		{
-			TSharedRef<SWidget> ThisWidget = WidgetToTest->GetChildren()->GetChildAt(ChildIdx);
-			const FText ChildText = FindChildWidgetWithText(ThisWidget.ToSharedPtr());
-			if (!ChildText.IsEmpty())
-			{
-				return ChildText;
-			}
-		}
+		FFormatNamedArguments Args;
+		Args.Add(TEXT("ItemName"), ItemName);
+		Args.Add(TEXT("ItemDescriptor"), ItemDescriptor);
+		FormattedItemName = FText::Format(LOCTEXT("ItemName_Format", " the \"{ItemName}\" {ItemDescriptor}"), Args);
 	}
-	
-	return FText();
-}
 
-// Return first active tooltip text found.  Return empty FText if not found.
-static FText FindCurrentToolTipText()
-{
-	for (const TArray<TSharedRef<SWindow>> AllWindows = FSlateApplication::Get().GetTopLevelWindows();
-		const TSharedRef<SWindow>& ThisWindow : AllWindows)
+	FText NameToQuery;
+	if (bIsInAssistantPanel)
 	{
-		if (ThisWindow->GetType() == EWindowType::ToolTip && ThisWindow->IsVisible())
-		{
-			return FindChildWidgetWithText(ThisWindow);
-		}
+		NameToQuery = LOCTEXT("ItemName_Assistant", "the AI Assistant");
+	}
+	else if (FormattedItemName.IsEmpty() && !TextUnderCursor.IsEmpty())
+	{
+		NameToQuery = TextUnderCursor;
+	}
+	else if (!FormattedItemName.IsEmpty() || !FormattedTabName.IsEmpty() || !FormattedWindowName.IsEmpty())
+	{
+		NameToQuery = (!FormattedItemName.IsEmpty()) ? FormattedItemName : (!FormattedTabName.IsEmpty()) ? FormattedTabName : FormattedWindowName;
 	}
-	
-	return FText();
-}
 
-// Search backwards from end of WidgetPath and return first SMenuEntryBlock or SWidgetBlock widget found.
-// Return null ptr if not found.
-static TSharedPtr<SWidget> FindClosestMenuItem(const FWidgetPath& WidgetPath)
-{
-	TSharedPtr<SWidget> MenuItemBlock;
-	if (WidgetPath.IsValid())
+	// Localized text is non-static const so it can change when the editor language is changed.
+	const FText QueryFormatWidget = LOCTEXT("QueryFormat1", "I would like to know what {NameToQuery} does.");
+	const FText QueryFormatGeneral = LOCTEXT("QueryFormat2", "I would like to know what \"{NameToQuery}\" means.");
+	const FText QueryFormatObject = LOCTEXT("QueryFormat3", "I would like to know about {NameToQuery}.");
+	const FText QueryFormatAssistant = LOCTEXT("QueryFormat4", "Tell me about yourself.");
+	// Hidden instruction text sent to assistant can be non-localized and static const.
+	static const FText QueryInstructionsWidgetFormat = INVTEXT(
+		"Provide an easily readable, informative, accurate answer that describes what the widget does in the Unreal Editor UI. {QueryAdditionalInstructions} {StyleInstructions}");
+	static const FText QueryInstructionsObjectFormat = INVTEXT(
+		"Provide an easily readable, informative, accurate answer that describes what I should know about it in the Unreal Editor UI. {QueryAdditionalInstructions} {StyleInstructions}");
+	static const FText QueryInstructionsAssistantFormat = INVTEXT(
+		"Provide an easily readable, informative, accurate answer that describes what I should know about yourself and what you can do for me. {StyleInstructions}");
+	static const FText QueryAdditionalInstructionsWidget = INVTEXT(
+		"Explain what happens if I select or click on it.");
+	static const FText QueryAdditionalInstructionsObject = INVTEXT(
+		"Describe what the object is or what it does.");
+	static const FText QueryAdditionalInstructionsGeneral = INVTEXT(
+		"If the object is clickable, like a button, then explain what happens if I click on it. If the object is something you select, like an asset or a blueprint node, describe what the object is or what it does.");
+	static const FText StyleInstructions = INVTEXT(
+		"Use bold text to emphasize Unreal Engine specific terms. Start with a quick overview paragraph, then add key points limited to one or two formatted sections with headers and bullet points. Don't include a summary at the end.");
+
+	if (!NameToQuery.IsEmpty())
 	{
-		// is it a menu item? Look up for an SMenuEntryBlock or SWidgetBlock.
-		static const TStaticArray<FName, 2> MenuTypes = {"SMenuEntryBlock", "SWidgetBlock"};
-		MenuItemBlock = FindClosestWidgetPtrOfSortedTypes(WidgetPath, MenuTypes);
+		FFormatNamedArguments QueryArgs;
+		QueryArgs.Add(TEXT("NameToQuery"), NameToQuery);
+		FFormatNamedArguments InstructionArgs;
 
-		// disabled menu items *contain* the menu entry block
-		if (!MenuItemBlock.IsValid())
+		if (bIsInAssistantPanel)
 		{
-			if (WidgetPath.GetLastWidget()->GetChildren()->Num() > 0)
-			{
-				for (int32 ChildIdx = 0; ChildIdx < WidgetPath.GetLastWidget()->GetChildren()->Num(); ChildIdx++)
-				{
-					TSharedRef<SWidget> ThisWidget = WidgetPath.GetLastWidget()->GetChildren()->GetChildAt(ChildIdx);
-					for (FName MenuType : MenuTypes)
-					{
-						if (ThisWidget->GetType() == MenuType)
-						{
-							MenuItemBlock = ThisWidget.ToSharedPtr();
-							break;
-						}
-					}
-				}
-			}
+			GeneratedQuery = FText::Format(QueryFormatAssistant, QueryArgs);
+			InstructionArgs.Add(TEXT("StyleInstructions"), StyleInstructions);
+			GeneratedQueryInstructions = FText::Format(QueryInstructionsAssistantFormat, InstructionArgs);
 		}
-	}
-
-	return MenuItemBlock;
-}
-
-// For menu items, we want to generate the WidgetPath context based on the root widget that spawned the menu,
-// not on the actual menu item itself.
-static FWidgetPath FindContextWidgetPath(const FWidgetPath& WidgetPath)
-{
-	const TSharedPtr<SWidget> MenuItem = FindClosestMenuItem(WidgetPath);
-	if (MenuItem.IsValid())
-	{
-		TSharedPtr<SWidget> RootWidget = FSlateApplication::Get().GetMenuHostWidget();
-		if (RootWidget.IsValid())
+		else if (bIsObject)
 		{
-			FWidgetPath OutWidgetPath;
-			FSlateApplication::Get().GeneratePathToWidgetChecked(RootWidget.ToSharedRef(), OutWidgetPath);
-			return OutWidgetPath;
+			GeneratedQuery = FText::Format(QueryFormatObject, QueryArgs);
+			InstructionArgs.Add(TEXT("QueryAdditionalInstructions"), QueryAdditionalInstructionsObject);
+			InstructionArgs.Add(TEXT("StyleInstructions"), StyleInstructions);
+			GeneratedQueryInstructions = FText::Format(QueryInstructionsObjectFormat, InstructionArgs);
 		}
-	}
-	return WidgetPath;
-}
-
-// Return true and fill in SlateQueryContext if queried item is in viewport.
-static bool ItemIsInViewport(const FWidgetPath& WidgetPath, FAIAssistantSlateQueryContext& SlateQueryContext)
-{
-	TSharedPtr<SWidget> Viewport = FindClosestWidgetPtrOfType(WidgetPath, "SViewport");
-	if (Viewport.IsValid())
-	{
-		// Web browser widgets also have Viewports!
-		// Ignore web browser views except for the special case of AI Assistant window or panel.
-		TSharedPtr<SWidget> WebBrowser = Viewport->GetParentWidget();
-		if (WebBrowser.IsValid() && WebBrowser->GetType() == "SWebBrowserView")
+		else if (bIsUIWidget)
 		{
-			TSharedPtr<SWidget> AssistantPane = FindClosestWidgetPtrOfType(WidgetPath, "SAIAssistantWebBrowser");
-			if (AssistantPane.IsValid())
-			{
-				SlateQueryContext.bIsInAssistantPanel = true;
-				SlateQueryContext.LastPickedWidget = AssistantPane;
-				return true;
-			}
+			GeneratedQuery = FText::Format(QueryFormatWidget, QueryArgs);
+			InstructionArgs.Add(TEXT("QueryAdditionalInstructions"), QueryAdditionalInstructionsWidget);
+			InstructionArgs.Add(TEXT("StyleInstructions"), StyleInstructions);
+			GeneratedQueryInstructions = FText::Format(QueryInstructionsWidgetFormat, InstructionArgs);
 		}
 		else
 		{
-			// Selected actors are filled in by the IsInEditor function.
-			if (!SlateQueryContext.SelectedActorName.IsEmpty())
-			{
-				SlateQueryContext.bIsObject = true;
-				SlateQueryContext.LastPickedWidget = Viewport;
-				SlateQueryContext.ItemName = FText::FromString(SlateQueryContext.SelectedActorName);
-				SlateQueryContext.ItemDescriptor = LOCTEXT("ItemDescriptor_SelectedActor", "actor");
-				SlateQueryContext.StructuredContext.ActorClass = SlateQueryContext.SelectedActorClass;
-				return true;
-			}
+			GeneratedQuery = FText::Format(QueryFormatWidget, QueryArgs);
+			InstructionArgs.Add(TEXT("QueryAdditionalInstructions"), QueryAdditionalInstructionsGeneral);
+			InstructionArgs.Add(TEXT("StyleInstructions"), StyleInstructions);
+			GeneratedQueryInstructions = FText::Format(QueryInstructionsWidgetFormat, InstructionArgs);
 		}
 	}
-	return false;
 }
 
-// Return true and fill in SlateQueryContext if queried item is a graph node.
-static bool ItemIsGraphNode(const FWidgetPath& WidgetPath, FAIAssistantSlateQueryContext& SlateQueryContext)
+void FSlateQuerier::GenerateStructuredContext()
 {
-	// Look up for an SGraphEditor and then down for an SGraphPanel. Whatever is the child of that SGraphPanel is our SGraphNode.
-	// Note that this rests on the assumption that SGraphPanel only contains children that are SGraphNodes. This assumption is also
-	// made in the code of SGraphPanel itself.
-	TSharedPtr<SGraphEditor> GraphEditor;
-	for (int32 WidgetIndex = WidgetPath.Widgets.Num() - 1; WidgetIndex >= 0; --WidgetIndex)
+	StructuredContext.UnrealEditorVersion = FEngineVersion::Current().ToString();
+	StructuredContext.UnrealEditorLanguage = FInternationalization::Get().GetCurrentLanguage()->GetDisplayName();
+	StructuredContext.UnrealEditorLocale = FInternationalization::Get().GetCurrentLocale()->GetName();
+
+	static const FText LanguageInstructionsFormat = INVTEXT(" Please respond in language \"{0}\". ");
+	const FText Language = FText::FromString(StructuredContext.UnrealEditorLanguage);
+	LanguageInstructions = FText::Format(LanguageInstructionsFormat, Language);
+
+	StructuredContext.ToolTipText = CurrentToolTipText.ToString();
+
+	if (!StructuredContext.ActorClass || StructuredContext.ActorClass->IsEmpty())
 	{
-		const FArrangedWidget* ArrangedWidget = &WidgetPath.Widgets[WidgetIndex];
-		const TSharedRef<SWidget>& ThisWidget = ArrangedWidget->Widget;
-		if (ThisWidget->GetType() == "SGraphEditor")
+		// If we're in a Details panel, include the class of actor or object that we're editing.
+		const TSharedPtr<SWidget> DetailsView = Utility::FindClosestWidgetPtrOfType(WidgetPath, "SDetailsView");
+		if (DetailsView.IsValid())
 		{
-			for (int32 EditorDescendantIndex = WidgetIndex; EditorDescendantIndex < WidgetPath.Widgets.Num(); EditorDescendantIndex++)
+			const TSharedPtr<IDetailsView> DetailsViewCast = StaticCastSharedPtr<IDetailsView>(DetailsView);
+			if (DetailsViewCast->GetSelectedActors().Num() > 0)
 			{
-				const FArrangedWidget* DescendantArrangedWidget = &WidgetPath.Widgets[EditorDescendantIndex];
-				const TSharedRef<SWidget>& ThisDescendantWidget = DescendantArrangedWidget->Widget;
-				if (ThisDescendantWidget->GetType() == "SGraphPanel" && (EditorDescendantIndex + 1 < WidgetPath.Widgets.Num()))
-				{
-					const FArrangedWidget* NodeArrangedWidget = &WidgetPath.Widgets[EditorDescendantIndex+1];
-					const TSharedRef<SWidget>& ThisNodeWidget = NodeArrangedWidget->Widget;
-					const TSharedPtr<SGraphNode> AsGraphNode = StaticCastSharedPtr<SGraphNode>(ThisNodeWidget.ToSharedPtr());
-					SlateQueryContext.LastPickedWidget = ThisNodeWidget.ToSharedPtr();
-					SlateQueryContext.ItemName = AsGraphNode->GetNodeObj()->GetNodeTitle(ENodeTitleType::MenuTitle);
-					SlateQueryContext.ItemDescriptor = LOCTEXT("ItemDescriptor_GraphNode", "graph node");
-					return true;
-				}
+				FString ActorName;
+				DetailsViewCast->GetSelectedActors()[0]->GetClass()->GetName(ActorName);
+				StructuredContext.ActorClass = ActorName;
+			}
+			else if (DetailsViewCast->GetSelectedObjects().Num() > 0)
+			{
+				FString ObjectName;
+				DetailsViewCast->GetSelectedObjects()[0]->GetClass()->GetName(ObjectName);
+				StructuredContext.ActorClass = ObjectName;
 			}
 		}
 	}
-	return false;
-}
 
-// Fill query context structure for buttons.
-static bool FillButtonQueryContext(const TSharedPtr<SWidget> Button,
-	FAIAssistantSlateQueryContext& SlateQueryContext, FText ItemDescriptor)
-{
-	FText ChildText = FindChildWidgetWithText(Button);
-	if (!ChildText.IsEmpty() || !SlateQueryContext.CurrentToolTipText.IsEmpty())
+	if (!ItemName.IsEmpty())
 	{
-		if (!ChildText.IsEmpty())
-		{
-			SlateQueryContext.ItemName = ChildText;
-		}
-		else
-		{
-			// The tool tip can be used in this case.
-			SlateQueryContext.ItemName = SlateQueryContext.CurrentToolTipText;
-		}
-		SlateQueryContext.ItemDescriptor = ItemDescriptor;
-		SlateQueryContext.LastPickedWidget = Button;
-		SlateQueryContext.bIsUIWidget = true;
-		return true;
+		StructuredContext.ItemName = ItemName.ToString();
+		StructuredContext.ItemDescriptor = ItemDescriptor.ToString();
 	}
-	return false;
-}
 
-// Return true and fill in SlateQueryContext if queried item is a button type as specified in
-// ButtonWidgetTypes.
-static bool ItemIsButtonType(const FWidgetPath& WidgetPath, FAIAssistantSlateQueryContext& SlateQueryContext,
-	const TSet<FName> ButtonWidgetTypes, FText ItemDescriptor)
-{
-	TSharedPtr<SWidget> Button = FindClosestWidgetPtrOfTypes(WidgetPath, ButtonWidgetTypes);
+	if (!TextUnderCursor.IsEmpty())
+	{
+		StructuredContext.TextUnderCursor = TextUnderCursor.ToString();
+	}
 
-	// Disabled buttons *contain* the menu entry block.
-	if (!Button.IsValid() && WidgetPath.GetLastWidget()->GetChildren()->Num() > 0)
+	// Get top level window title.
+	TSharedRef<SWindow> WidgetWindow = WidgetPath.GetWindow();
+	FText WindowTitle = WidgetWindow->GetTitle();
+	if (!WindowTitle.IsEmpty())
 	{
-		for (int32 ChildIdx = 0; ChildIdx < WidgetPath.GetLastWidget()->GetChildren()->Num(); ChildIdx++)
+		StructuredContext.TopWindowName = WindowTitle.ToString();
+	}
+
+	// Get picked widget type.
+	if (LastPickedWidget.IsValid())
+	{
+		StructuredContext.WidgetType = LastPickedWidget->GetType().ToString();
+	}
+
+	if (ItemName.IsEmpty() && !bIsInAssistantPanel)
+	{
+		StructuredContext.WidgetPath = TArray<FString>();
+		// Add simplified widget path
+		for (int32 WidgetIndex = 0; WidgetIndex < WidgetPath.Widgets.Num(); WidgetIndex++)
 		{
-			TSharedRef<SWidget> ThisWidget = WidgetPath.GetLastWidget()->GetChildren()->GetChildAt(ChildIdx);
-			if (ButtonWidgetTypes.Contains(ThisWidget->GetType()))
+			const FArrangedWidget* ArrangedWidget = &WidgetPath.Widgets[WidgetIndex];
+			const TSharedRef<SWidget>& ThisWidget = ArrangedWidget->Widget;
+			FString NextWidgetString = ThisWidget->GetType().ToString();
+			// Skip uninteresting widgets
+			if (NextWidgetString != TEXT("SBorder") && NextWidgetString != TEXT("SBox") &&
+				NextWidgetString != TEXT("SScrollBorder") &&
+				NextWidgetString != TEXT("SOverlay") && NextWidgetString != TEXT("SVerticalBox") &&
+				NextWidgetString != TEXT("SHorizontalBox") && NextWidgetString != TEXT("SSplitter") &&
+				NextWidgetString != TEXT("SDockingSplitter"))
 			{
-				Button = ThisWidget.ToSharedPtr();
+				StructuredContext.WidgetPath->Add(NextWidgetString);
 			}
 		}
 	}
-
-	return Button.IsValid() && FillButtonQueryContext(Button, SlateQueryContext, ItemDescriptor);
 }
 
-// Return true and fill in SlateQueryContext if queried item is a button.
-static bool ItemIsButton(const FWidgetPath& WidgetPath, FAIAssistantSlateQueryContext& SlateQueryContext)
+void FSlateQuerier::SendQuery()
 {
-	static const TSet<FName> ToolbarButtonTypes = TSet<FName>({
-		"SToolBarButtonBlock", "SToolBarComboButtonBlock"
-		});
+	TSharedPtr<SAIAssistantWebBrowser> WebBrowser =
+		UAIAssistantSubsystem::GetAIAssistantWebBrowserWidget();
+	WebBrowser->CreateConversation();
 
-	static const TSet<FName> SimpleButtonTypes = TSet<FName>({
-			"SButton",
-			"SPrimaryButton",
-			"SCheckBox"
-		});
+	const FString VisiblePromptString =
+		Utility::CleanExtraWhiteSpaceFromString(GeneratedQuery.ToString());
+
+	const FString HiddenContextString =
+		Utility::CleanExtraWhiteSpaceFromString(GeneratedQueryInstructions.ToString()) +
+		LanguageInstructions.ToString() +
+		TEXT(" ContextData=") +
+		StructuredContext.ToJson(false);
 
-	return (ItemIsButtonType(WidgetPath, SlateQueryContext, ToolbarButtonTypes,
-				LOCTEXT("ItemDescriptor_ToolbarButton", "toolbar button")) ||
-			ItemIsButtonType(WidgetPath, SlateQueryContext, SimpleButtonTypes,
-				LOCTEXT("ItemDescriptor_Button", "button")));
+	WebBrowser->AddUserMessageToConversation(VisiblePromptString, HiddenContextString);
 }
 
-// Return true and fill in SlateQueryContext if queried item is a search box.
-static bool ItemIsSearchBox(const FWidgetPath& WidgetPath, FAIAssistantSlateQueryContext& SlateQueryContext)
+bool FSlateQuerier::FindEditorName()
 {
-	static const TSet<FName> SearchBoxTypes = TSet<FName>({"SSearchBox", "SFilterSearchBox"});
-	TSharedPtr<SWidget> SearchBox = FindClosestWidgetPtrOfTypes(WidgetPath, SearchBoxTypes);
+	FText OutName = FText();
 
-	if (SearchBox.IsValid())
+	if (!WidgetPath.IsValid())
 	{
-		FText HintText;
-		if (SearchBox->GetType() == "SSearchBox")
-		{
-			TSharedPtr<SSearchBox> ThisSearchBoxCast = StaticCastSharedPtr<SSearchBox>(SearchBox);
-			HintText = ThisSearchBoxCast->GetHintText();
-		}
-		else
-		{
-			TSharedPtr<SWidget> EditableText = FindChildWidgetOfType(SearchBox, "SEditableText");
-			if (EditableText)
-			{
-				TSharedPtr<SEditableText> ThisEditableTextCast = StaticCastSharedPtr<SEditableText>(EditableText);
-				HintText = ThisEditableTextCast->GetHintText();
-			}
-		}
-		SlateQueryContext.ItemName = FText::Format(LOCTEXT("SearchBoxFormat", "{0}"),
-			(HintText.IsEmpty()) ? LOCTEXT("SearchBoxDefaultName", "default") : HintText);
-		SlateQueryContext.ItemDescriptor = LOCTEXT("ItemDescriptor_SearchBox", "search box");
-		SlateQueryContext.LastPickedWidget = SearchBox;
-		SlateQueryContext.bIsUIWidget = true;
+		FormattedWindowName = OutName;
+		return false;
+	}
+
+	// is it in the status bar?
+	TSharedPtr<SWidget> StatusBar = Utility::FindClosestWidgetPtrOfType(WidgetPath, "SStatusBar");
+	if (StatusBar.IsValid())
+	{
+		LastPickedWidget = StatusBar;
+		StructuredContext.TabName = TEXT("StatusBar");
+		FormattedWindowName = LOCTEXT("TabName_StatusBar", " the Status Bar");
 		return true;
 	}
 
-	return false;
-}
+	if (IsInEditor())
+	{
+		return true;
+	}
 
-// Return true and fill in SlateQueryContext if queried item is a console input box.
-static bool ItemIsConsoleInputBox(const FWidgetPath& WidgetPath, FAIAssistantSlateQueryContext& SlateQueryContext)
-{
-	TSharedPtr<SWidget> ConsoleInputBox = FindClosestWidgetPtrOfType(WidgetPath, "SConsoleInputBox");
-	if (ConsoleInputBox.IsValid())
+	// Is it in a drawer overlay?
+	if (IsInDrawer())
 	{
-		FText HintText;
-		TSharedPtr<SWidget> EditableText = FindChildWidgetOfType(ConsoleInputBox, "SMultiLineEditableTextBox");
-		if (EditableText.IsValid())
-		{
-			TSharedPtr<SMultiLineEditableTextBox> ThisMultiLineEditableTextBoxCast = StaticCastSharedPtr<SMultiLineEditableTextBox>(EditableText);
-			HintText = ThisMultiLineEditableTextBoxCast->GetHintText();
-		}
-		SlateQueryContext.ItemName = FText::Format(LOCTEXT("ConsoleInputBoxFormat", "{0}"),
-			(HintText.IsEmpty() ? LOCTEXT("ConsoleInputBoxDefaultName","Console") : HintText));
-		SlateQueryContext.ItemDescriptor = LOCTEXT("ItemDescriptor_ConsoleInputBox", "input box");
-		SlateQueryContext.LastPickedWidget = ConsoleInputBox;
-		SlateQueryContext.bIsUIWidget = true;
 		return true;
 	}
-	return false;
-}
 
-// Return true and fill in SlateQueryContext if queried item is a menu item.
-static bool ItemIsMenuItem(const FWidgetPath& WidgetPath, FAIAssistantSlateQueryContext& SlateQueryContext)
-{
-	// Look up for an SMenuEntryBlock or SWidgetBlock and then down for the text.
-	TSharedPtr<SWidget> MenuItemBlock = FindClosestMenuItem(WidgetPath);
-	if (MenuItemBlock.IsValid())
+	TSharedRef<SWindow> WidgetWindow = WidgetPath.GetWindow();
+	// is it the main window that hosts the level editor?
+	if (FModuleManager::Get().IsModuleLoaded("LevelEditor"))
 	{
-		FText ChildText = FindChildWidgetWithText(MenuItemBlock);
-		if (!ChildText.IsEmpty())
+		IMainFrameModule& MainFrameModule = FModuleManager::LoadModuleChecked<IMainFrameModule>(TEXT("MainFrame"));
+		if (WidgetWindow.ToSharedPtr() == MainFrameModule.GetParentWindow())
 		{
-			SlateQueryContext.ItemName = ChildText;
-			SlateQueryContext.ItemDescriptor = LOCTEXT("ItemDescriptor_Menu", "menu item");
-			SlateQueryContext.LastPickedWidget = MenuItemBlock;
-			SlateQueryContext.bIsUIWidget = true;
+			StructuredContext.WindowName = TEXT("MainFrame");
+			FormattedWindowName = LOCTEXT("WindowMainFrame_Title", " the editor's main window");
 			return true;
 		}
 	}
+
+	// does the window have a title, like the Project Browser or Color Picker?
+	FText WindowTitle = WidgetWindow->GetTitle();
+	if (!WindowTitle.IsEmpty())
+	{
+		StructuredContext.WindowName = WindowTitle.ToString();
+		LastPickedWidget = WidgetWindow.ToSharedPtr();
+		FFormatNamedArguments Args;
+		Args.Add(TEXT("WindowTitle"), WindowTitle);
+		FormattedWindowName = FText::Format(LOCTEXT("WindowNameFormat_Title", " the {WindowTitle} window"), Args);
+		return true;
+	}
+
+	// TODO: how to find other editor contexts?
 	return false;
 }
 
-// Return true and fill in SlateQueryContext if queried item is a Details panel property.
-static bool ItemIsDetailsProperty(const FWidgetPath& WidgetPath, FAIAssistantSlateQueryContext& SlateQueryContext)
+bool FSlateQuerier::FindTabName()
 {
-	// Look up for an SDetailSingleItemRow [could be others?] then down through the SPropertyNameWidget.
-	TSharedPtr<SWidget> PropertyRow = FindClosestWidgetPtrOfType(WidgetPath, "SDetailSingleItemRow");
-	if (PropertyRow.IsValid())
-	{
-		FText ChildText;
+	bool bFoundTabName = false;
 
-		SlateQueryContext.LastPickedWidget = PropertyRow;
-
-		TSharedPtr<SWidget> PropertyNameWidget = FindChildWidgetOfType(PropertyRow, "SPropertyNameWidget");
-		if (PropertyNameWidget.IsValid())
-		{
-			ChildText = FindChildWidgetWithText(PropertyNameWidget);
-			if (!ChildText.IsEmpty())
-			{
-				SlateQueryContext.ItemName = ChildText;
-				SlateQueryContext.ItemDescriptor = LOCTEXT("ItemDescriptor_Property", "setting");
-			}
-		}
-		else
-		{
-			// Some rows e.g. Transform have no SPropertyNameWidget. Best we can do is look for a text label.
-			ChildText = FindChildWidgetWithText(PropertyRow);
-			if (!ChildText.IsEmpty())
-			{
-				SlateQueryContext.ItemName = ChildText;
-				SlateQueryContext.ItemDescriptor = LOCTEXT("ItemDescriptor_Property", "setting");
-			}
-		}
-
-		// Find section name from sibling SDetailCategoryTableRow.
-		if (PropertyRow->IsParentValid())
-		{
-			TSharedPtr<SWidget> ListPanel = PropertyRow->GetParentWidget();
-			TSharedRef<SWidget> PreviousDetailCategory = SNullWidget::NullWidget;
-			for (int32 ChildIdx = 0; ChildIdx < ListPanel->GetChildren()->Num(); ChildIdx++)
-			{
-				TSharedRef<SWidget> ThisWidget = ListPanel->GetChildren()->GetChildAt(ChildIdx);
-				if (ThisWidget == PropertyRow)
-				{
-					break;
-				}
-				if (ThisWidget->GetType() == "SDetailCategoryTableRow")
-				{
-					PreviousDetailCategory = ThisWidget;
-				}
-			}
-			if (PreviousDetailCategory->GetType() != "SNullWidgetContent")
-			{
-				ChildText = FindChildWidgetWithText(PreviousDetailCategory);
-				if (!ChildText.IsEmpty())
-				{
-					SlateQueryContext.StructuredContext.CategoryName = ChildText.ToString();
-					SlateQueryContext.ItemDescriptor = FText::Format(
-						LOCTEXT("ItemDescriptor_PropertyWithCategory", "setting in the {0} category"), ChildText);
-				}
-			}
-		}
-		return true;
-	}
-	return false;
-}
-
-// Return true and fill in SlateQueryContext if queried item is a plugin in Plugins browser.
-static bool ItemIsPlugin(const FWidgetPath& WidgetPath, FAIAssistantSlateQueryContext& SlateQueryContext)
-{
-	// Look up for containing SPluginBrowser.
-	TSharedPtr<SWidget> PluginBrowser = FindClosestWidgetPtrOfType(WidgetPath, "SPluginBrowser");
-	if (PluginBrowser.IsValid())
-	{
-		FText ChildText;
-		TSharedPtr<SWidget> PluginCategoryTree = FindClosestWidgetPtrOfType(WidgetPath, "SPluginCategoryTree");
-		TSharedPtr<SWidget> PluginTileList = FindClosestWidgetPtrOfType(WidgetPath, "SPluginTileList");
-		if (PluginCategoryTree.IsValid())
-		{
-			SlateQueryContext.LastPickedWidget = PluginCategoryTree;
-			TSharedPtr<SWidget> PluginCategory = FindClosestWidgetPtrOfType(WidgetPath, "SPluginCategory");
-			if (PluginCategory.IsValid())
-			{
-				SlateQueryContext.LastPickedWidget = PluginCategory;
-				ChildText = FindChildWidgetWithText(PluginCategory);
-				if (!ChildText.IsEmpty())
-				{
-					SlateQueryContext.ItemName = ChildText;
-					SlateQueryContext.ItemDescriptor = LOCTEXT("ItemDescriptor_PluginCategory", "plugin category");
-					SlateQueryContext.bIsObject = true;
-					return true;
-				}
-			}
-		}
-		else if (PluginTileList.IsValid())
-		{
-			SlateQueryContext.LastPickedWidget = PluginTileList;
-			TSharedPtr<SWidget> PluginTile = FindClosestWidgetPtrOfType(WidgetPath, "SPluginTile");
-			if (PluginTile.IsValid())
-			{
-				SlateQueryContext.LastPickedWidget = PluginTile;
-				ChildText = FindChildWidgetWithText(PluginTile);
-				if (!ChildText.IsEmpty())
-				{
-					SlateQueryContext.ItemName = ChildText;
-					SlateQueryContext.ItemDescriptor = LOCTEXT("ItemDescriptor_PluginTile", "plugin");
-					SlateQueryContext.bIsObject = true;
-					return true;
-				}
-			}
-		}
-	}
-	return false;
-}
-
-// Return true and fill in SlateQueryContext if queried item is a navigation breadcrumb button.
-static bool ItemIsBreadcrumbTrailButton(const FWidgetPath& WidgetPath, FAIAssistantSlateQueryContext& SlateQueryContext)
-{
-	TSharedPtr<SWidget> BreadcrumbTrail = FindClosestWidgetPtrOfType(WidgetPath, "SBreadcrumbTrail<FNavigationCrumb>");
-	TSharedPtr<SWidget> BreadcrumbButton = BreadcrumbTrail ? FindClosestWidgetPtrOfType(WidgetPath, "SButton") : nullptr;
-	FText ChildText = FindChildWidgetWithText(BreadcrumbButton);
-	if (!ChildText.IsEmpty())
-	{
-		SlateQueryContext.ItemName = ChildText;
-		SlateQueryContext.ItemDescriptor = LOCTEXT("ItemDescriptor_Breadcrumb", "navigation breadcrumb");
-		SlateQueryContext.LastPickedWidget = BreadcrumbButton;
-		SlateQueryContext.bIsUIWidget = true;
-		return true;
-	}
-	return false;
-}
-
-// Return true and fill in SlateQueryContext if queried item is an asset item.
-static bool ItemIsAsset(const FWidgetPath& WidgetPath, FAIAssistantSlateQueryContext& SlateQueryContext)
-{
-	FText ChildText;
-
-	// Is it an asset tile item?
-	TSharedPtr<SWidget> AssetTileItem = FindClosestWidgetPtrOfType(WidgetPath, "SAssetTileItem");
-	if (AssetTileItem.IsValid())
-	{
-		TSharedPtr<SWidget> AssetThumbnail = FindChildWidgetOfType(AssetTileItem, "SAssetThumbnail");
-		if (AssetThumbnail.IsValid())
-		{
-			ChildText = FindChildWidgetWithText(AssetThumbnail); // This returns TYPE of asset instead of name
-			if (!ChildText.IsEmpty())
-			{
-				FFormatNamedArguments Args;
-				Args.Add(TEXT("AssetType"), ChildText);
-				SlateQueryContext.ItemDescriptor = FText::Format(LOCTEXT("ItemDescriptor_Asset", "{AssetType} asset"), Args);
-			}
-
-			if (!SlateQueryContext.CurrentToolTipText.IsEmpty())
-			{
-				// We can use the tooltip for the asset's name.
-				SlateQueryContext.ItemName = SlateQueryContext.CurrentToolTipText;
-			}
-			SlateQueryContext.LastPickedWidget = AssetThumbnail;
-			SlateQueryContext.bIsObject = true;
-		}
-		else
-		{
-			ChildText = FindChildWidgetWithText(AssetTileItem);
-			if (!ChildText.IsEmpty())
-			{
-				SlateQueryContext.ItemName = ChildText;
-				SlateQueryContext.ItemDescriptor = LOCTEXT("ItemDescriptor_Folder", "asset folder");
-				SlateQueryContext.LastPickedWidget = AssetTileItem;
-				SlateQueryContext.bIsObject = true;
-			}
-		}
-		return true;
-	}
-
-	// Is it an asset tree item?
-	TSharedPtr<SWidget> AssetTreeItem = FindClosestWidgetPtrOfType(WidgetPath, "SAssetTreeItem");
-	if (AssetTreeItem.IsValid())
-	{
-		ChildText = FindChildWidgetWithText(AssetTreeItem);
-		if (!ChildText.IsEmpty())
-		{
-			SlateQueryContext.ItemName = ChildText;
-			SlateQueryContext.ItemDescriptor = LOCTEXT("ItemDescriptor_Folder", "asset folder");
-			SlateQueryContext.LastPickedWidget = AssetTreeItem;
-			SlateQueryContext.bIsObject = true;
-			return true;
-		}
-	}
-	return false;
-}
-
-// Return true and fill in SlateQueryContext if queried item is in Outliner or Subobject Instance Editor.
-static bool ItemIsOutlinerActorOrSubobject(const FWidgetPath& WidgetPath, FAIAssistantSlateQueryContext& SlateQueryContext)
-{
-	FText ChildText;
-
-	// Is it an actor in the Outliner?
-	TSharedPtr<SWidget> Outliner = FindClosestWidgetPtrOfType(WidgetPath, "SSceneOutliner");
-	TSharedPtr<SWidget> ActorTreeLabel = Outliner ? FindClosestWidgetPtrOfType(WidgetPath, "SActorTreeLabel") : nullptr;
-	TSharedPtr<SWidget> OutlinerRow = Outliner ? FindClosestWidgetPtrOfType(WidgetPath, "SSceneOutlinerTreeRow") : nullptr;
-
-	if (!ActorTreeLabel.IsValid() && OutlinerRow.IsValid())
-	{
-		ActorTreeLabel = FindChildWidgetOfType(OutlinerRow, "SActorTreeLabel");
-	}
-
-	if (ActorTreeLabel.IsValid())
-	{
-		ChildText = FindChildWidgetWithText(ActorTreeLabel);
-		if (!ChildText.IsEmpty())
-		{
-			SlateQueryContext.ItemName = ChildText;
-			SlateQueryContext.ItemDescriptor = LOCTEXT("ItemDescriptor_OutlinerActor", "actor");
-			SlateQueryContext.LastPickedWidget = ActorTreeLabel;
-			SlateQueryContext.bIsObject = true;
-			return true;
-		}
-	}
-
-	// Is it an object in the Subobject Instance editor?
-	TSharedPtr<SWidget> SubobjectInstanceEditor = FindClosestWidgetPtrOfType(WidgetPath, "SSubobjectInstanceEditor");
-	TSharedPtr<SWidget> SubObjectRow = SubobjectInstanceEditor ? FindClosestWidgetPtrOfType(WidgetPath, "SSubobject_RowWidget") : nullptr;
-	TSharedPtr<SWidget> NameWidget = SubObjectRow ? FindChildWidgetOfType(SubObjectRow, "SInlineEditableTextBlock") : nullptr;
-
-	if (NameWidget.IsValid())
-	{
-		ChildText = FindChildWidgetWithText(NameWidget);
-		SlateQueryContext.ItemName = ChildText;
-		SlateQueryContext.ItemDescriptor = LOCTEXT("ItemDescriptor_SubobjectInstance", "instance");
-		SlateQueryContext.LastPickedWidget = SubObjectRow;
-		SlateQueryContext.bIsObject = true;
-		return true;
-	}
-
-	return false;
-}
-
-// Return true and fill in SlateQueryContext if queried item is a Details panel filter checkbox.
-static bool ItemIsDetailPanelFilterCheckbox(const FWidgetPath& WidgetPath, FAIAssistantSlateQueryContext& SlateQueryContext)
-{
-	TSharedPtr<SWidget> ActorDetails = FindClosestWidgetPtrOfType(WidgetPath, "SActorDetails");
-	if (!ActorDetails)
-	{
-		return false;
-	}
-
-	TSharedPtr<SWidget> DetailsView = FindClosestWidgetPtrOfType(WidgetPath, "SDetailsView");
-	TSharedPtr<SWidget> DetailsNameArea = FindClosestWidgetPtrOfType(WidgetPath, "SDetailsNameArea");
-	TSharedPtr<SWidget> SubobjectInstanceEditor = FindClosestWidgetPtrOfType(WidgetPath, "SSubobjectInstanceEditor");
-	TSharedPtr<SWidget> WrapBox = FindClosestWidgetPtrOfType(WidgetPath, "SWrapBox");
-	TSharedPtr<SWidget> CheckBox = FindClosestWidgetPtrOfType(WidgetPath, "SCheckBox");
-
-	if (!DetailsView && !DetailsNameArea && !SubobjectInstanceEditor && WrapBox && CheckBox)
-	{
-		return FillButtonQueryContext(CheckBox, SlateQueryContext,
-			LOCTEXT("ItemDescriptor_PropertyFilter", "property filter"));
-	}
-	return false;
-}
-
-// Return true and fill in SlateQueryContext with item information if valid item information is found.
-static bool FindItemName(const FWidgetPath& WidgetPath, FAIAssistantSlateQueryContext& SlateQueryContext)
-{
-	SlateQueryContext.ItemDescriptor = LOCTEXT("ItemDescriptor_Generic", "control");
-
-	if (!WidgetPath.IsValid())
-	{
-		return false;
-	}
-
-	if (// Is it a graph node?
-		ItemIsGraphNode(WidgetPath, SlateQueryContext) ||
-		// Is it a search/filter field?
-		ItemIsSearchBox(WidgetPath, SlateQueryContext) ||
-		// Is it a console input box?
-		ItemIsConsoleInputBox(WidgetPath, SlateQueryContext) ||
-		// Is it a details panel property?
-		ItemIsDetailsProperty(WidgetPath, SlateQueryContext) ||
-		// Is it a plugin in the plugins browser ?
-		ItemIsPlugin(WidgetPath, SlateQueryContext) ||
-		// Is it a breadcrumb trail button?
-		ItemIsBreadcrumbTrailButton(WidgetPath, SlateQueryContext) ||
-		// Is it an asset item?
-		ItemIsAsset(WidgetPath, SlateQueryContext) ||
-		// Is it an Actor in Outliner, or an instance or component in the SubobjectInstanceEditor?
-		ItemIsOutlinerActorOrSubobject(WidgetPath, SlateQueryContext) ||
-		// Does the widget path end up in a Viewport?
-		ItemIsInViewport(WidgetPath, SlateQueryContext) ||
-		// Is it one of the Details panel checkboxes that filter the detail properties viewed?
-		ItemIsDetailPanelFilterCheckbox(WidgetPath, SlateQueryContext) ||
-		// Is it a menu item?
-		ItemIsMenuItem(WidgetPath, SlateQueryContext) ||
-		// Is it a button?
-		ItemIsButton(WidgetPath, SlateQueryContext))
-	{
-		return true;
-	}
-
-	return false;
-}
-
-// Return true and fill in SlateQueryContext if valid Tab or Panel information for queried item is found.
-static bool FindTabName(const FWidgetPath& WidgetPath, FAIAssistantSlateQueryContext& SlateQueryContext)
-{
-	bool bFoundTabName = false;
-
-	TSharedPtr<SWidget> TabStack = FindClosestWidgetPtrOfType(WidgetPath, "SDockingTabStack");
+	TSharedPtr<SWidget> TabStack = Utility::FindClosestWidgetPtrOfType(WidgetPath, "SDockingTabStack");
 	if (TabStack.IsValid())
 	{
 		TArray<TSharedRef<SWidget>> DockTabs;
-		FindChildWidgetsOfType(DockTabs, TabStack, "SDockTab");
+		Utility::FindChildWidgetsOfType(DockTabs, TabStack, "SDockTab");
 		for (auto& ThisTab : DockTabs)
 		{
 			TSharedRef<SDockTab> ThisTabCast = StaticCastSharedRef<SDockTab>(ThisTab);
@@ -851,11 +374,11 @@ static bool FindTabName(const FWidgetPath& WidgetPath, FAIAssistantSlateQueryCon
 			}
 			if (ThisTabCast->IsForeground())
 			{
-				SlateQueryContext.LastPickedWidget = TabStack;
-				SlateQueryContext.StructuredContext.TabName = ThisTabCast->GetTabLabel().ToString();
+				LastPickedWidget = TabStack;
+				StructuredContext.TabName = ThisTabCast->GetTabLabel().ToString();
 				FFormatNamedArguments Args;
 				Args.Add(TEXT("PanelName"), ThisTabCast->GetTabLabel());
-				SlateQueryContext.FormattedTabName = FText::Format(LOCTEXT("TabName_Format", " the {PanelName} panel"), Args);
+				FormattedTabName = FText::Format(LOCTEXT("TabName_Format", " the {PanelName} panel"), Args);
 				bFoundTabName = true;
 			}
 		}
@@ -864,30 +387,7 @@ static bool FindTabName(const FWidgetPath& WidgetPath, FAIAssistantSlateQueryCon
 	return bFoundTabName;
 }
 
-// If we're in a details panel, include the class of actor or object that we're editing in SlateQueryContext.
-static void GenerateDetailsViewContext(const FWidgetPath& WidgetPath, FAIAssistantSlateQueryContext& SlateQueryContext)
-{
-	const TSharedPtr<SWidget> DetailsView = FindClosestWidgetPtrOfType(WidgetPath, "SDetailsView");
-	if (DetailsView.IsValid())
-	{
-		const TSharedPtr<IDetailsView> DetailsViewCast = StaticCastSharedPtr<IDetailsView>(DetailsView);
-		if (DetailsViewCast->GetSelectedActors().Num() > 0)
-		{
-			FString ActorName;
-			DetailsViewCast->GetSelectedActors()[0]->GetClass()->GetName(ActorName);
-			SlateQueryContext.StructuredContext.ActorClass = ActorName;
-		}
-		else if (DetailsViewCast->GetSelectedObjects().Num() > 0)
-		{
-			FString ObjectName;
-			DetailsViewCast->GetSelectedObjects()[0]->GetClass()->GetName(ObjectName);
-			SlateQueryContext.StructuredContext.ActorClass = ObjectName;
-		}
-	}
-}
-
-// Return true and fill in SlateQueryContext with item information if valid editor information is found.
-static bool IsInEditor(const FWidgetPath& WidgetPath, FAIAssistantSlateQueryContext& SlateQueryContext)
+bool FSlateQuerier::IsInEditor()
 {
 	static const TSet<FEditorModeID> BlockedModeIDs = TSet<FEditorModeID>({
 		FBuiltinEditorModes::EM_Default,
@@ -895,10 +395,10 @@ static bool IsInEditor(const FWidgetPath& WidgetPath, FAIAssistantSlateQueryCont
 		});
 
 	// are we in the level editor?
-	TSharedPtr<SWidget> LevelEditor = FindFirstWidgetPtrOfType(WidgetPath, "SLevelEditor");
+	TSharedPtr<SWidget> LevelEditor = Utility::FindFirstWidgetPtrOfType(WidgetPath, "SLevelEditor");
 	if (LevelEditor)
 	{
-		SlateQueryContext.StructuredContext.EditorName = TEXT("LevelEditor");
+		StructuredContext.EditorName = TEXT("LevelEditor");
 		FText ModeString;
 		ULevelEditorSubsystem* LevelEditorSubsystem = GEditor->GetEditorSubsystem<ULevelEditorSubsystem>();
 		if (FEditorModeTools* ModeManager = LevelEditorSubsystem->GetLevelEditorModeManager())
@@ -915,7 +415,7 @@ static bool IsInEditor(const FWidgetPath& WidgetPath, FAIAssistantSlateQueryCont
 						FText ToolDisplayName = ModeToolkit.Pin()->GetActiveToolDisplayName();
 						if (!ToolDisplayName.IsEmpty())
 						{
-							SlateQueryContext.StructuredContext.ActiveToolName = ToolDisplayName.ToString();
+							StructuredContext.ActiveToolName = ToolDisplayName.ToString();
 						}
 					}
 				}
@@ -924,20 +424,20 @@ static bool IsInEditor(const FWidgetPath& WidgetPath, FAIAssistantSlateQueryCont
 					FFormatNamedArguments Args;
 					Args.Add(TEXT("ModeName"), ModeInfo.Name);
 					ModeString = FText::Format(LOCTEXT("EditorName_ModeFormat", "with {ModeName} mode active"), Args);
-					SlateQueryContext.StructuredContext.EditorMode = ModeInfo.Name.ToString();
+					StructuredContext.EditorMode = ModeInfo.Name.ToString();
 				}
 			}
 		}
-		SlateQueryContext.LastPickedWidget = LevelEditor;
+		LastPickedWidget = LevelEditor;
 		if (ModeString.IsEmpty())
 		{
-			SlateQueryContext.FormattedWindowName = LOCTEXT("EditorName_LevelEditor", " the Level Editor");
+			FormattedWindowName = LOCTEXT("EditorName_LevelEditor", " the Level Editor");
 		}
 		else
 		{
 			FFormatNamedArguments Args;
 			Args.Add(TEXT("WithModeActive"), ModeString);
-			SlateQueryContext.FormattedWindowName = FText::Format(LOCTEXT("EditorName_LevelEditorAndMode", " the Level Editor {WithModeActive}"), Args);
+			FormattedWindowName = FText::Format(LOCTEXT("EditorName_LevelEditorAndMode", " the Level Editor {WithModeActive}"), Args);
 		}
 
 		if (!GEditor->IsPlaySessionInProgress())
@@ -951,15 +451,15 @@ static bool IsInEditor(const FWidgetPath& WidgetPath, FAIAssistantSlateQueryCont
 			{
 				if (Actor)
 				{
-					SlateQueryContext.SelectedActorName = Actor->GetName();
-					SlateQueryContext.SelectedActorClass = Actor->GetClass()->GetName();
+					SelectedActorName = Actor->GetName();
+					SelectedActorClass = Actor->GetClass()->GetName();
 					bActorIsSelected = true;
 					break;
 				}
 			}
 
 			// If there is no selection and mouse is in viewport, see what actor is under mouse cursor.
-			if (!bActorIsSelected && FindClosestWidgetPtrOfType(WidgetPath, "SViewport"))
+			if (!bActorIsSelected && Utility::FindClosestWidgetPtrOfType(WidgetPath, "SViewport"))
 			{
 				FViewport* ActiveViewport = GEditor->GetActiveViewport();
 				if (ActiveViewport)
@@ -973,8 +473,8 @@ static bool IsInEditor(const FWidgetPath& WidgetPath, FAIAssistantSlateQueryCont
 						AActor* TargetActor = TargetProxy->Actor;
 						if (TargetActor)
 						{
-							SlateQueryContext.SelectedActorName = TargetActor->GetName();
-							SlateQueryContext.SelectedActorClass = TargetActor->GetClass()->GetName();
+							SelectedActorName = TargetActor->GetName();
+							SelectedActorClass = TargetActor->GetClass()->GetName();
 							bActorIsSelected = true;
 						}
 					}
@@ -986,15 +486,15 @@ static bool IsInEditor(const FWidgetPath& WidgetPath, FAIAssistantSlateQueryCont
 	}
 
 	// if not, are we in a different asset editor?
-	TSharedPtr<SWidget> StandaloneAssetEditorHost = FindFirstWidgetPtrOfType(WidgetPath, "SStandaloneAssetEditorToolkitHost");
+	TSharedPtr<SWidget> StandaloneAssetEditorHost = Utility::FindFirstWidgetPtrOfType(WidgetPath, "SStandaloneAssetEditorToolkitHost");
 	if (StandaloneAssetEditorHost)
 	{
 		TArray<TSharedRef<SWidget>> TabStacks;
-		FindChildWidgetsOfType(TabStacks, StandaloneAssetEditorHost, "SDockingTabStack");
+		Utility::FindChildWidgetsOfType(TabStacks, StandaloneAssetEditorHost, "SDockingTabStack");
 		for (auto& ThisTabStack : TabStacks)
 		{
 			TArray<TSharedRef<SWidget>> DockTabs;
-			FindChildWidgetsOfType(DockTabs, ThisTabStack.ToSharedPtr(), "SDockTab");
+			Utility::FindChildWidgetsOfType(DockTabs, ThisTabStack.ToSharedPtr(), "SDockTab");
 
 			for (auto& ThisTab : DockTabs)
 			{
@@ -1012,15 +512,15 @@ static bool IsInEditor(const FWidgetPath& WidgetPath, FAIAssistantSlateQueryCont
 							IAssetEditorInstance* ThisEditor = AssetEditorSubsystem->FindEditorForAsset(ThisEditedAsset, false);
 							if (ThisEditor && ThisEditor->GetAssociatedTabManager() == PickedWidgetTabManager)
 							{
-								SlateQueryContext.StructuredContext.EditorName = ThisEditor->GetEditorName().ToString();
-								SlateQueryContext.LastPickedWidget = ThisTabStack.ToSharedPtr();
+								StructuredContext.EditorName = ThisEditor->GetEditorName().ToString();
+								LastPickedWidget = ThisTabStack.ToSharedPtr();
 								FFormatNamedArguments Args;
 								Args.Add(TEXT("AssetEditorName"), FText::FromName(ThisEditor->GetEditorName()));
-								SlateQueryContext.FormattedWindowName = FText::Format(
+								FormattedWindowName = FText::Format(
 									LOCTEXT("EditorName_AssetEditorFormat", " the {AssetEditorName} editor"), Args);
 
-								SlateQueryContext.SelectedActorName = ThisEditedAsset->GetName();
-								SlateQueryContext.SelectedActorClass = ThisEditedAsset->GetClass()->GetName();
+								SelectedActorName = ThisEditedAsset->GetName();
+								SelectedActorClass = ThisEditedAsset->GetClass()->GetName();
 
 								return true;
 							}
@@ -1033,18 +533,16 @@ static bool IsInEditor(const FWidgetPath& WidgetPath, FAIAssistantSlateQueryCont
 	return false;
 }
 
-// Return true and fill in SlateQueryContext with item information if queried item is in a drawer and valid
-// window and tab information is found.
-static bool IsInDrawer(const FWidgetPath& WidgetPath, FAIAssistantSlateQueryContext& SlateQueryContext)
+bool FSlateQuerier::IsInDrawer()
 {
-	TSharedPtr<SWidget> DrawerOverlay = FindFirstWidgetPtrOfType(WidgetPath, "SDrawerOverlay");
+	TSharedPtr<SWidget> DrawerOverlay = Utility::FindFirstWidgetPtrOfType(WidgetPath, "SDrawerOverlay");
 	if (DrawerOverlay.IsValid())
 	{
 		FText DrawerName;
-		TSharedPtr<SWidget> DrawerWidget = FindChildWidgetOfType(DrawerOverlay, "SContentBrowser");
+		TSharedPtr<SWidget> DrawerWidget = Utility::FindChildWidgetOfType(DrawerOverlay, "SContentBrowser");
 		if (!DrawerWidget)
 		{
-			DrawerWidget = FindChildWidgetOfType(DrawerOverlay, "SOutputLog");
+			DrawerWidget = Utility::FindChildWidgetOfType(DrawerOverlay, "SOutputLog");
 		}
 
 		if (DrawerWidget)
@@ -1057,16 +555,16 @@ static bool IsInDrawer(const FWidgetPath& WidgetPath, FAIAssistantSlateQueryCont
 			{
 				DrawerName = LOCTEXT("DrawerName_OutputLog", "OutputLog");
 			}
-			// Are there other items that are commonly in drawers?
+			// TODO: Are there other items that are commonly in drawers?
 
 			if (!DrawerName.IsEmpty())
 			{
-				SlateQueryContext.StructuredContext.WindowName = DrawerOverlay->GetType().ToString();
-				SlateQueryContext.StructuredContext.TabName = DrawerName.ToString();
-				SlateQueryContext.LastPickedWidget = DrawerWidget;
+				StructuredContext.WindowName = DrawerOverlay->GetType().ToString();
+				StructuredContext.TabName = DrawerName.ToString();
+				LastPickedWidget = DrawerWidget;
 				FFormatNamedArguments Args;
 				Args.Add(TEXT("DrawerName"), DrawerName);
-				SlateQueryContext.FormattedWindowName = FText::Format(LOCTEXT("DrawerNameFormat", "the {DrawerName} drawer"), Args);
+				FormattedWindowName = FText::Format(LOCTEXT("DrawerNameFormat", "the {DrawerName} drawer"), Args);
 				return true;
 			}
 		}
@@ -1074,330 +572,494 @@ static bool IsInDrawer(const FWidgetPath& WidgetPath, FAIAssistantSlateQueryCont
 	return false;
 }
 
-// Return true and fill in SlateQueryContext with item information if valid editor or window information is found.
-static bool FindEditorName(const FWidgetPath& WidgetPath, FAIAssistantSlateQueryContext& SlateQueryContext)
+bool FSlateQuerier::FindItemName()
 {
-	FText OutName = FText();
+	ItemDescriptor = LOCTEXT("ItemDescriptor_Generic", "control");
 
 	if (!WidgetPath.IsValid())
 	{
-		SlateQueryContext.FormattedWindowName = OutName;
 		return false;
 	}
 
-	// is it in the status bar?
-	TSharedPtr<SWidget> StatusBar = FindClosestWidgetPtrOfType(WidgetPath, "SStatusBar");
-	if (StatusBar.IsValid())
-	{
-		SlateQueryContext.LastPickedWidget = StatusBar;
-		SlateQueryContext.StructuredContext.TabName = TEXT("StatusBar");
-		SlateQueryContext.FormattedWindowName = LOCTEXT("TabName_StatusBar", " the Status Bar");
-		return true;
-	}
-
-	if (IsInEditor(WidgetPath, SlateQueryContext))
-	{
-		return true;
-	}
-
-	// Is it in a drawer overlay?
-	if (IsInDrawer(WidgetPath, SlateQueryContext))
-	{
-		return true;
-	}
-
-	TSharedRef<SWindow> WidgetWindow = WidgetPath.GetWindow();
-	// is it the main window that hosts the level editor?
-	if (FModuleManager::Get().IsModuleLoaded("LevelEditor"))
-	{
-		IMainFrameModule& MainFrameModule = FModuleManager::LoadModuleChecked<IMainFrameModule>(TEXT("MainFrame"));
-		if (WidgetWindow.ToSharedPtr() == MainFrameModule.GetParentWindow())
-		{
-			SlateQueryContext.StructuredContext.WindowName = TEXT("MainFrame");
-			SlateQueryContext.FormattedWindowName = LOCTEXT("WindowMainFrame_Title", " the editor's main window");
-			return true;
-		}
-	}
-	
-	// does the window have a title, like the Project Browser or Color Picker?
-	FText WindowTitle = WidgetWindow->GetTitle();
-	if (!WindowTitle.IsEmpty())
+	if (// Is it a graph node?
+		ItemIsGraphNode() ||
+		// Is it a search/filter field?
+		ItemIsSearchBox() ||
+		// Is it a console input box?
+		ItemIsConsoleInputBox() ||
+		// Is it a details panel property?
+		ItemIsDetailsProperty() ||
+		// Is it a plugin in the plugins browser ?
+		ItemIsPlugin() ||
+		// Is it a breadcrumb trail button?
+		ItemIsBreadcrumbTrailButton() ||
+		// Is it an asset item?
+		ItemIsAsset() ||
+		// Is it an Actor in Outliner, or an instance or component in the SubobjectInstanceEditor?
+		ItemIsOutlinerActorOrSubobject() ||
+		// Is it one of the Details panel checkboxes that filter the detail properties viewed?
+		ItemIsDetailPanelFilterCheckbox() ||
+		// Does the widget path end up in a Viewport?
+		ItemIsInViewport() ||
+		// Is it a menu item?
+		ItemIsMenuItem() ||
+		// Is it a button?
+		ItemIsButton())
 	{
-		SlateQueryContext.StructuredContext.WindowName = WindowTitle.ToString();
-		SlateQueryContext.LastPickedWidget = WidgetWindow.ToSharedPtr();
-		FFormatNamedArguments Args;
-		Args.Add(TEXT("WindowTitle"), WindowTitle);
-		SlateQueryContext.FormattedWindowName = FText::Format(LOCTEXT("WindowNameFormat_Title", " the {WindowTitle} window"), Args);
 		return true;
 	}
 
-	// TODO: how to find other editor contexts?
 	return false;
 }
 
-// Fill in AccumulatedSlateQueryContext with information on text under cursor if found.
-static void FindTextUnderCursor(const FWidgetPath& WidgetPath, FAIAssistantSlateQueryContext& AccumulatedSlateQueryContext)
+bool FSlateQuerier::ItemIsInViewport()
 {
-	FText OutText;
-
-	if (WidgetPath.IsValid())
+	TSharedPtr<SWidget> Viewport = Utility::FindClosestWidgetPtrOfType(WidgetPath, "SViewport");
+	if (Viewport.IsValid())
 	{
-		const FRegexPattern AlphaPattern(TEXT("[A-Za-z]+"));
-		TSharedRef<SWidget> WidgetToTest = WidgetPath.GetLastWidget();
-		FText WidgetText;
-		if (WidgetToTest->GetType() == "STextBlock")
-		{
-			TSharedRef<STextBlock> TextBlock = StaticCastSharedRef<STextBlock>(WidgetToTest);
-			WidgetText = TextBlock->GetText();
-		}
-		else if (WidgetToTest->GetType() == "SRichTextBlock")
-		{
-			TSharedRef<SRichTextBlock> TextBlock = StaticCastSharedRef<SRichTextBlock>(WidgetToTest);
-			WidgetText = TextBlock->GetText();
-		}
-		if (!WidgetText.IsEmpty())
+		// Web browser widgets also have Viewports!
+		// Ignore web browser views except for the special case of AI Assistant window or panel.
+		TSharedPtr<SWidget> WebBrowser = Viewport->GetParentWidget();
+		if (WebBrowser.IsValid() && WebBrowser->GetType() == "SWebBrowserView")
 		{
-			FRegexMatcher AlphaMatcher(AlphaPattern, WidgetText.ToString());
-			if (AlphaMatcher.FindNext())
+			TSharedPtr<SWidget> AssistantPane = Utility::FindClosestWidgetPtrOfType(WidgetPath, "SAIAssistantWebBrowser");
+			if (AssistantPane.IsValid())
 			{
-				OutText = WidgetText;
+				bIsInAssistantPanel = true;
+				LastPickedWidget = AssistantPane;
+				return true;
 			}
 		}
-
-		if (!OutText.IsEmpty())
+		else
 		{
-			AccumulatedSlateQueryContext.LastPickedWidget = WidgetToTest;
-			AccumulatedSlateQueryContext.TextUnderCursor = OutText;
-			AccumulatedSlateQueryContext.StructuredContext.TextUnderCursor = OutText.ToString();
+			// Selected actors are filled in by the IsInEditor function.
+			if (!SelectedActorName.IsEmpty())
+			{
+				bIsObject = true;
+				LastPickedWidget = Viewport;
+				ItemName = FText::FromString(SelectedActorName);
+				ItemDescriptor = LOCTEXT("ItemDescriptor_SelectedActor", "actor");
+				StructuredContext.ActorClass = SelectedActorClass;
+				return true;
+			}
 		}
 	}
+	return false;
 }
 
-
-static FString CleanExtraWhiteSpaceFromString(const FString& InString)
+bool FSlateQuerier::ItemIsGraphNode()
 {
-	FString CleanString = "";
-	static const FRegexPattern WhiteSpacePattern(TEXT("(\\S+)"));
-	FRegexMatcher Matcher(WhiteSpacePattern, InString);
-	while (Matcher.FindNext())
+	// Look up for an SGraphEditor and then down for an SGraphPanel. Whatever is the child of that SGraphPanel is our SGraphNode.
+	// Note that this rests on the assumption that SGraphPanel only contains children that are SGraphNodes. This assumption is also
+	// made in the code of SGraphPanel itself.
+	TSharedPtr<SGraphEditor> GraphEditor;
+	for (int32 WidgetIndex = WidgetPath.Widgets.Num() - 1; WidgetIndex >= 0; --WidgetIndex)
 	{
-		if (!CleanString.IsEmpty())
+		const FArrangedWidget* ArrangedWidget = &WidgetPath.Widgets[WidgetIndex];
+		const TSharedRef<SWidget>& ThisWidget = ArrangedWidget->Widget;
+		if (ThisWidget->GetType() == "SGraphEditor")
 		{
-			CleanString += " ";
+			for (int32 EditorDescendantIndex = WidgetIndex; EditorDescendantIndex < WidgetPath.Widgets.Num(); EditorDescendantIndex++)
+			{
+				const FArrangedWidget* DescendantArrangedWidget = &WidgetPath.Widgets[EditorDescendantIndex];
+				const TSharedRef<SWidget>& ThisDescendantWidget = DescendantArrangedWidget->Widget;
+				if (ThisDescendantWidget->GetType() == "SGraphPanel" && (EditorDescendantIndex + 1 < WidgetPath.Widgets.Num()))
+				{
+					const FArrangedWidget* NodeArrangedWidget = &WidgetPath.Widgets[EditorDescendantIndex+1];
+					const TSharedRef<SWidget>& ThisNodeWidget = NodeArrangedWidget->Widget;
+					const TSharedPtr<SGraphNode> AsGraphNode = StaticCastSharedPtr<SGraphNode>(ThisNodeWidget.ToSharedPtr());
+					LastPickedWidget = ThisNodeWidget.ToSharedPtr();
+					ItemName = AsGraphNode->GetNodeObj()->GetNodeTitle(ENodeTitleType::MenuTitle);
+					ItemDescriptor = LOCTEXT("ItemDescriptor_GraphNode", "graph node");
+					return true;
+				}
+			}
 		}
-		
-		CleanString += Matcher.GetCaptureGroup(1);
 	}
-
-	return MoveTemp(CleanString);
+	return false;
 }
 
-//
-// UE::AIAssistant::SlateQuerier
-//
-FWidgetPath UE::AIAssistant::SlateQuerier::GetWidgetPathUnderCursor()
+bool FSlateQuerier::ItemIsSearchBox()
 {
-	const FVector2f CursorPos = FSlateApplication::Get().GetCursorPos();
-	const TArray<TSharedRef<SWindow>>& Windows = FSlateApplication::Get().GetTopLevelWindows();
-	return FSlateApplication::Get().LocateWindowUnderMouse(CursorPos, Windows, true);
-}
+	static const TSet<FName> SearchBoxTypes = TSet<FName>({"SSearchBox", "SFilterSearchBox"});
+	TSharedPtr<SWidget> SearchBox = Utility::FindClosestWidgetPtrOfTypes(WidgetPath, SearchBoxTypes);
 
-void UE::AIAssistant::SlateQuerier::QueryAIAssistantAboutSlateWidget(const FWidgetPath& InWidgetPath)
-{
-	if (!InWidgetPath.IsValid())
+	if (SearchBox.IsValid())
 	{
-		UE_LOG(LogAIAssistant, Warning, TEXT("No valid widget path found to generate widget query."));
-		return;
+		FText HintText;
+		if (SearchBox->GetType() == "SSearchBox")
+		{
+			TSharedPtr<SSearchBox> ThisSearchBoxCast = StaticCastSharedPtr<SSearchBox>(SearchBox);
+			HintText = ThisSearchBoxCast->GetHintText();
+		}
+		else
+		{
+			TSharedPtr<SWidget> EditableText = Utility::FindChildWidgetOfType(SearchBox, "SEditableText");
+			if (EditableText)
+			{
+				TSharedPtr<SEditableText> ThisEditableTextCast = StaticCastSharedPtr<SEditableText>(EditableText);
+				HintText = ThisEditableTextCast->GetHintText();
+			}
+		}
+		ItemName = FText::Format(LOCTEXT("SearchBoxFormat", "{0}"),
+			(HintText.IsEmpty()) ? LOCTEXT("SearchBoxDefaultName", "default") : HintText);
+		ItemDescriptor = LOCTEXT("ItemDescriptor_SearchBox", "search box");
+		LastPickedWidget = SearchBox;
+		bIsUIWidget = true;
+		return true;
 	}
 
-	// We build this up, below.
-	
-	FAIAssistantSlateQueryContext SlateQueryContext;
-
-	// Is there a current tool-tip?
-	SlateQueryContext.CurrentToolTipText = FindCurrentToolTipText();
-
-	// For menu items, we want to generate the context based on the root widget that spawned the menu,
-	// not on the actual menu item itself.
-	const FWidgetPath ContextWidgetPath = FindContextWidgetPath(InWidgetPath);
-
-	// Is there an identifiable asset editor, mode, or window?
-	FindEditorName(ContextWidgetPath, SlateQueryContext);
-
-	// Is there an identifiable tab?
-	// Look up the chain for an SDockingTabStack, then look down to find an SDockTab.
-	FindTabName(ContextWidgetPath, SlateQueryContext);
-
-	// Is there text under the cursor??
-	FindTextUnderCursor(InWidgetPath, SlateQueryContext);
-
-	// Is there an identifiable item?
-	FindItemName(InWidgetPath, SlateQueryContext);
+	return false;
+}
 
-	// Construct query.
+bool FSlateQuerier::ItemIsConsoleInputBox()
+{
+	TSharedPtr<SWidget> ConsoleInputBox = Utility::FindClosestWidgetPtrOfType(WidgetPath, "SConsoleInputBox");
+	if (ConsoleInputBox.IsValid())
 	{
-		FText FormattedItemName = SlateQueryContext.ItemName;
-		if (!FormattedItemName.IsEmpty())
+		FText HintText;
+		TSharedPtr<SWidget> EditableText = Utility::FindChildWidgetOfType(ConsoleInputBox, "SMultiLineEditableTextBox");
+		if (EditableText.IsValid())
 		{
-			FFormatNamedArguments Args;
-			Args.Add(TEXT("ItemName"), SlateQueryContext.ItemName);
-			Args.Add(TEXT("ItemDescriptor"), SlateQueryContext.ItemDescriptor);
-			FormattedItemName = FText::Format(LOCTEXT("ItemName_Format", " the \"{ItemName}\" {ItemDescriptor}"), Args);
+			TSharedPtr<SMultiLineEditableTextBox> ThisMultiLineEditableTextBoxCast = StaticCastSharedPtr<SMultiLineEditableTextBox>(EditableText);
+			HintText = ThisMultiLineEditableTextBoxCast->GetHintText();
 		}
+		ItemName = FText::Format(LOCTEXT("ConsoleInputBoxFormat", "{0}"),
+			(HintText.IsEmpty() ? LOCTEXT("ConsoleInputBoxDefaultName","Console") : HintText));
+		ItemDescriptor = LOCTEXT("ItemDescriptor_ConsoleInputBox", "input box");
+		LastPickedWidget = ConsoleInputBox;
+		bIsUIWidget = true;
+		return true;
+	}
+	return false;
+}
 
-		FText NameToQuery;
-		if (SlateQueryContext.bIsInAssistantPanel)
+bool FSlateQuerier::ItemIsDetailsProperty()
+{
+	// Look up for an SDetailSingleItemRow [could be others?] then down through the SPropertyNameWidget.
+	TSharedPtr<SWidget> PropertyRow = Utility::FindClosestWidgetPtrOfType(WidgetPath, "SDetailSingleItemRow");
+	if (PropertyRow.IsValid())
+	{
+		FText ChildText;
+
+		LastPickedWidget = PropertyRow;
+
+		TSharedPtr<SWidget> PropertyNameWidget = Utility::FindChildWidgetOfType(PropertyRow, "SPropertyNameWidget");
+		if (PropertyNameWidget.IsValid())
 		{
-			NameToQuery = LOCTEXT("ItemName_Assistant", "the AI Assistant");
+			ChildText = Utility::FindChildWidgetWithText(PropertyNameWidget);
+			if (!ChildText.IsEmpty())
+			{
+				ItemName = ChildText;
+				ItemDescriptor = LOCTEXT("ItemDescriptor_Property", "setting");
+			}
 		}
-		else if (FormattedItemName.IsEmpty() && !SlateQueryContext.TextUnderCursor.IsEmpty())
+		else
 		{
-			NameToQuery = SlateQueryContext.TextUnderCursor;
+			// Some rows e.g. Transform have no SPropertyNameWidget. Best we can do is look for a text label.
+			ChildText = Utility::FindChildWidgetWithText(PropertyRow);
+			if (!ChildText.IsEmpty())
+			{
+				ItemName = ChildText;
+				ItemDescriptor = LOCTEXT("ItemDescriptor_Property", "setting");
+			}
 		}
-		else if (!FormattedItemName.IsEmpty() || !SlateQueryContext.FormattedTabName.IsEmpty() || !SlateQueryContext.FormattedWindowName.IsEmpty())
+
+		// Find section name from sibling SDetailCategoryTableRow.
+		if (PropertyRow->IsParentValid())
 		{
-			NameToQuery = (!FormattedItemName.IsEmpty()) ? FormattedItemName : (!SlateQueryContext.FormattedTabName.IsEmpty()) ? SlateQueryContext.FormattedTabName : SlateQueryContext.FormattedWindowName;
+			TSharedPtr<SWidget> ListPanel = PropertyRow->GetParentWidget();
+			TSharedRef<SWidget> PreviousDetailCategory = SNullWidget::NullWidget;
+			for (int32 ChildIdx = 0; ChildIdx < ListPanel->GetChildren()->Num(); ChildIdx++)
+			{
+				TSharedRef<SWidget> ThisWidget = ListPanel->GetChildren()->GetChildAt(ChildIdx);
+				if (ThisWidget == PropertyRow)
+				{
+					break;
+				}
+				if (ThisWidget->GetType() == "SDetailCategoryTableRow")
+				{
+					PreviousDetailCategory = ThisWidget;
+				}
+			}
+			if (PreviousDetailCategory->GetType() != "SNullWidgetContent")
+			{
+				ChildText = Utility::FindChildWidgetWithText(PreviousDetailCategory);
+				if (!ChildText.IsEmpty())
+				{
+					StructuredContext.CategoryName = ChildText.ToString();
+					ItemDescriptor = FText::Format(
+						LOCTEXT("ItemDescriptor_PropertyWithCategory", "setting in the {0} category"), ChildText);
+				}
+			}
 		}
+		return true;
+	}
+	return false;
+}
 
-		// Localized text is non-static const so it can change when the editor language is changed.
-		const FText QueryFormatWidget = LOCTEXT("QueryFormat1", "I would like to know what {NameToQuery} does.");
-		const FText QueryFormatGeneral = LOCTEXT("QueryFormat2", "I would like to know what \"{NameToQuery}\" means.");
-		const FText QueryFormatObject = LOCTEXT("QueryFormat3", "I would like to know about {NameToQuery}.");
-		const FText QueryFormatAssistant = LOCTEXT("QueryFormat4", "Tell me about yourself.");
-		// Hidden instruction text sent to assistant can be non-localized and static const.
-		static const FText QueryInstructionsWidgetFormat = INVTEXT(
-			"Provide an easily readable, informative, accurate answer that describes what the widget does in the Unreal Editor UI. {QueryAdditionalInstructions} {StyleInstructions}");
-		static const FText QueryInstructionsObjectFormat = INVTEXT(
-			"Provide an easily readable, informative, accurate answer that describes what I should know about it in the Unreal Editor UI. {QueryAdditionalInstructions} {StyleInstructions}");
-		static const FText QueryInstructionsAssistantFormat = INVTEXT(
-			"Provide an easily readable, informative, accurate answer that describes what I should know about yourself and what you can do for me. {StyleInstructions}");
-		static const FText QueryAdditionalInstructionsWidget = INVTEXT(
-			"Explain what happens if I select or click on it.");
-		static const FText QueryAdditionalInstructionsObject = INVTEXT(
-			"Describe what the object is or what it does.");
-		static const FText QueryAdditionalInstructionsGeneral = INVTEXT(
-			"If the object is clickable, like a button, then explain what happens if I click on it. If the object is something you select, like an asset or a blueprint node, describe what the object is or what it does.");
-		static const FText StyleInstructions = INVTEXT(
-			"Use bold text to emphasize Unreal Engine specific terms. Start with a quick overview paragraph, then add key points limited to one or two formatted sections with headers and bullet points. Don't include a summary at the end.");
-
-		if (!NameToQuery.IsEmpty())
+bool FSlateQuerier::ItemIsPlugin()
+{
+	// Look up for containing SPluginBrowser.
+	TSharedPtr<SWidget> PluginBrowser = Utility::FindClosestWidgetPtrOfType(WidgetPath, "SPluginBrowser");
+	if (PluginBrowser.IsValid())
+	{
+		FText ChildText;
+		TSharedPtr<SWidget> PluginCategoryTree = Utility::FindClosestWidgetPtrOfType(WidgetPath, "SPluginCategoryTree");
+		TSharedPtr<SWidget> PluginTileList = Utility::FindClosestWidgetPtrOfType(WidgetPath, "SPluginTileList");
+		if (PluginCategoryTree.IsValid())
 		{
-			FFormatNamedArguments QueryArgs;
-			QueryArgs.Add(TEXT("NameToQuery"), NameToQuery);
-			FFormatNamedArguments InstructionArgs;
-
-			if (SlateQueryContext.bIsInAssistantPanel)
+			LastPickedWidget = PluginCategoryTree;
+			TSharedPtr<SWidget> PluginCategory = Utility::FindClosestWidgetPtrOfType(WidgetPath, "SPluginCategory");
+			if (PluginCategory.IsValid())
 			{
-				SlateQueryContext.GeneratedQuery = FText::Format(QueryFormatAssistant, QueryArgs);
-				InstructionArgs.Add(TEXT("StyleInstructions"), StyleInstructions);
-				SlateQueryContext.GeneratedQueryInstructions = FText::Format(QueryInstructionsAssistantFormat, InstructionArgs);
+				LastPickedWidget = PluginCategory;
+				ChildText = Utility::FindChildWidgetWithText(PluginCategory);
+				if (!ChildText.IsEmpty())
+				{
+					ItemName = ChildText;
+					ItemDescriptor = LOCTEXT("ItemDescriptor_PluginCategory", "plugin category");
+					bIsObject = true;
+					return true;
+				}
 			}
-			else if (SlateQueryContext.bIsObject)
+		}
+		else if (PluginTileList.IsValid())
+		{
+			LastPickedWidget = PluginTileList;
+			TSharedPtr<SWidget> PluginTile = Utility::FindClosestWidgetPtrOfType(WidgetPath, "SPluginTile");
+			if (PluginTile.IsValid())
 			{
-				SlateQueryContext.GeneratedQuery = FText::Format(QueryFormatObject, QueryArgs);
-				InstructionArgs.Add(TEXT("QueryAdditionalInstructions"), QueryAdditionalInstructionsObject);
-				InstructionArgs.Add(TEXT("StyleInstructions"), StyleInstructions);
-				SlateQueryContext.GeneratedQueryInstructions = FText::Format(QueryInstructionsObjectFormat, InstructionArgs);
+				LastPickedWidget = PluginTile;
+				ChildText = Utility::FindChildWidgetWithText(PluginTile);
+				if (!ChildText.IsEmpty())
+				{
+					ItemName = ChildText;
+					ItemDescriptor = LOCTEXT("ItemDescriptor_PluginTile", "plugin");
+					bIsObject = true;
+					return true;
+				}
 			}
-			else if (SlateQueryContext.bIsUIWidget)
+		}
+	}
+	return false;
+}
+
+bool FSlateQuerier::ItemIsBreadcrumbTrailButton()
+{
+	TSharedPtr<SWidget> BreadcrumbTrail = Utility::FindClosestWidgetPtrOfType(WidgetPath, "SBreadcrumbTrail<FNavigationCrumb>");
+	TSharedPtr<SWidget> BreadcrumbButton = BreadcrumbTrail ? Utility::FindClosestWidgetPtrOfType(WidgetPath, "SButton") : nullptr;
+	FText ChildText = Utility::FindChildWidgetWithText(BreadcrumbButton);
+	if (!ChildText.IsEmpty())
+	{
+		ItemName = ChildText;
+		ItemDescriptor = LOCTEXT("ItemDescriptor_Breadcrumb", "navigation breadcrumb");
+		LastPickedWidget = BreadcrumbButton;
+		bIsUIWidget = true;
+		return true;
+	}
+	return false;
+}
+
+bool FSlateQuerier::ItemIsAsset()
+{
+	FText ChildText;
+
+	// Is it an asset tile item?
+	TSharedPtr<SWidget> AssetTileItem = Utility::FindClosestWidgetPtrOfType(WidgetPath, "SAssetTileItem");
+	if (AssetTileItem.IsValid())
+	{
+		TSharedPtr<SWidget> AssetThumbnail = Utility::FindChildWidgetOfType(AssetTileItem, "SAssetThumbnail");
+		if (AssetThumbnail.IsValid())
+		{
+			ChildText = Utility::FindChildWidgetWithText(AssetThumbnail); // This returns TYPE of asset instead of name
+			if (!ChildText.IsEmpty())
 			{
-				SlateQueryContext.GeneratedQuery = FText::Format(QueryFormatWidget, QueryArgs);
-				InstructionArgs.Add(TEXT("QueryAdditionalInstructions"), QueryAdditionalInstructionsWidget);
-				InstructionArgs.Add(TEXT("StyleInstructions"), StyleInstructions);
-				SlateQueryContext.GeneratedQueryInstructions = FText::Format(QueryInstructionsWidgetFormat, InstructionArgs);
+				FFormatNamedArguments Args;
+				Args.Add(TEXT("AssetType"), ChildText);
+				ItemDescriptor = FText::Format(LOCTEXT("ItemDescriptor_Asset", "{AssetType} asset"), Args);
 			}
-			else
+
+			if (!CurrentToolTipText.IsEmpty())
 			{
-				SlateQueryContext.GeneratedQuery = FText::Format(QueryFormatWidget, QueryArgs);
-				InstructionArgs.Add(TEXT("QueryAdditionalInstructions"), QueryAdditionalInstructionsGeneral);
-				InstructionArgs.Add(TEXT("StyleInstructions"), StyleInstructions);
-				SlateQueryContext.GeneratedQueryInstructions = FText::Format(QueryInstructionsWidgetFormat, InstructionArgs);
+				// We can use the tooltip for the asset's name.
+				ItemName = CurrentToolTipText;
 			}
+			LastPickedWidget = AssetThumbnail;
+			bIsObject = true;
 		}
-
-		if (SlateQueryContext.GeneratedQuery.IsEmpty())
+		else
 		{
-			UE_LOG(LogAIAssistant, Warning, TEXT("Could not generate query for widget."));
-			return;
+			ChildText = Utility::FindChildWidgetWithText(AssetTileItem);
+			if (!ChildText.IsEmpty())
+			{
+				ItemName = ChildText;
+				ItemDescriptor = LOCTEXT("ItemDescriptor_Folder", "asset folder");
+				LastPickedWidget = AssetTileItem;
+				bIsObject = true;
+			}
 		}
+		return true;
 	}
 
-	// Construct structured context.
+	// Is it an asset tree item?
+	TSharedPtr<SWidget> AssetTreeItem = Utility::FindClosestWidgetPtrOfType(WidgetPath, "SAssetTreeItem");
+	if (AssetTreeItem.IsValid())
 	{
-		SlateQueryContext.StructuredContext.UnrealEditorVersion = FEngineVersion::Current().ToString();
-		SlateQueryContext.StructuredContext.UnrealEditorLanguage = FInternationalization::Get().GetCurrentLanguage()->GetDisplayName();
-		SlateQueryContext.StructuredContext.UnrealEditorLocale = FInternationalization::Get().GetCurrentLocale()->GetName();
+		ChildText = Utility::FindChildWidgetWithText(AssetTreeItem);
+		if (!ChildText.IsEmpty())
+		{
+			ItemName = ChildText;
+			ItemDescriptor = LOCTEXT("ItemDescriptor_Folder", "asset folder");
+			LastPickedWidget = AssetTreeItem;
+			bIsObject = true;
+			return true;
+		}
+	}
+	return false;
+}
 
-		static const FText LanguageInstructionsFormat = INVTEXT(" Please respond in language \"{0}\". ");
-		const FText Language = FText::FromString(SlateQueryContext.StructuredContext.UnrealEditorLanguage);
-		SlateQueryContext.LanguageInstructions = FText::Format(LanguageInstructionsFormat, Language);
+bool FSlateQuerier::ItemIsOutlinerActorOrSubobject()
+{
+	FText ChildText;
 
-		SlateQueryContext.StructuredContext.ToolTipText = SlateQueryContext.CurrentToolTipText.ToString();
+	// Is it an actor in the Outliner?
+	TSharedPtr<SWidget> Outliner = Utility::FindClosestWidgetPtrOfType(WidgetPath, "SSceneOutliner");
+	TSharedPtr<SWidget> ActorTreeLabel = Outliner ? Utility::FindClosestWidgetPtrOfType(WidgetPath, "SActorTreeLabel") : nullptr;
+	TSharedPtr<SWidget> OutlinerRow = Outliner ? Utility::FindClosestWidgetPtrOfType(WidgetPath, "SSceneOutlinerTreeRow") : nullptr;
 
-		GenerateDetailsViewContext(InWidgetPath, SlateQueryContext);
+	if (!ActorTreeLabel.IsValid() && OutlinerRow.IsValid())
+	{
+		ActorTreeLabel = Utility::FindChildWidgetOfType(OutlinerRow, "SActorTreeLabel");
+	}
 
-		if (!SlateQueryContext.ItemName.IsEmpty())
+	if (ActorTreeLabel.IsValid())
+	{
+		ChildText = Utility::FindChildWidgetWithText(ActorTreeLabel);
+		if (!ChildText.IsEmpty())
 		{
-			SlateQueryContext.StructuredContext.ItemName = SlateQueryContext.ItemName.ToString();
-			SlateQueryContext.StructuredContext.ItemDescriptor = SlateQueryContext.ItemDescriptor.ToString();
+			ItemName = ChildText;
+			ItemDescriptor = LOCTEXT("ItemDescriptor_OutlinerActor", "actor");
+			LastPickedWidget = ActorTreeLabel;
+			bIsObject = true;
+			return true;
 		}
+	}
 
-		// Get top level window title.
-		TSharedRef<SWindow> WidgetWindow = InWidgetPath.GetWindow();
-		FText WindowTitle = WidgetWindow->GetTitle();
-		if (!WindowTitle.IsEmpty())
-		{
-			SlateQueryContext.StructuredContext.TopWindowName = WindowTitle.ToString();
-		}
+	// Is it an object in the Subobject Instance editor?
+	TSharedPtr<SWidget> SubobjectInstanceEditor = Utility::FindClosestWidgetPtrOfType(WidgetPath, "SSubobjectInstanceEditor");
+	TSharedPtr<SWidget> SubObjectRow = SubobjectInstanceEditor ? Utility::FindClosestWidgetPtrOfType(WidgetPath, "SSubobject_RowWidget") : nullptr;
+	TSharedPtr<SWidget> NameWidget = SubObjectRow ? Utility::FindChildWidgetOfType(SubObjectRow, "SInlineEditableTextBlock") : nullptr;
 
-		// Get picked widget type.
-		if (SlateQueryContext.LastPickedWidget.IsValid())
-		{
-			SlateQueryContext.StructuredContext.WidgetType = SlateQueryContext.LastPickedWidget->GetType().ToString();
-		}
+	if (NameWidget.IsValid())
+	{
+		ChildText = Utility::FindChildWidgetWithText(NameWidget);
+		ItemName = ChildText;
+		ItemDescriptor = LOCTEXT("ItemDescriptor_SubobjectInstance", "instance");
+		LastPickedWidget = SubObjectRow;
+		bIsObject = true;
+		return true;
+	}
 
-		if (SlateQueryContext.ItemName.IsEmpty() && !SlateQueryContext.bIsInAssistantPanel)
+	return false;
+}
+
+bool FSlateQuerier::ItemIsDetailPanelFilterCheckbox()
+{
+	TSharedPtr<SWidget> ActorDetails = Utility::FindClosestWidgetPtrOfType(WidgetPath, "SActorDetails");
+	if (!ActorDetails)
+	{
+		return false;
+	}
+
+	TSharedPtr<SWidget> DetailsView = Utility::FindClosestWidgetPtrOfType(WidgetPath, "SDetailsView");
+	TSharedPtr<SWidget> DetailsNameArea = Utility::FindClosestWidgetPtrOfType(WidgetPath, "SDetailsNameArea");
+	TSharedPtr<SWidget> SubobjectInstanceEditor = Utility::FindClosestWidgetPtrOfType(WidgetPath, "SSubobjectInstanceEditor");
+	TSharedPtr<SWidget> WrapBox = Utility::FindClosestWidgetPtrOfType(WidgetPath, "SWrapBox");
+	TSharedPtr<SWidget> CheckBox = Utility::FindClosestWidgetPtrOfType(WidgetPath, "SCheckBox");
+
+	if (!DetailsView && !DetailsNameArea && !SubobjectInstanceEditor && WrapBox && CheckBox)
+	{
+		return SetItemAsButton(CheckBox, LOCTEXT("ItemDescriptor_PropertyFilter", "property filter"));
+	}
+	return false;
+}
+
+bool FSlateQuerier::ItemIsMenuItem()
+{
+	// Look up for an SMenuEntryBlock or SWidgetBlock and then down for the text.
+	TSharedPtr<SWidget> MenuItemBlock = Utility::FindClosestMenuItem(WidgetPath);
+	if (MenuItemBlock.IsValid())
+	{
+		FText ChildText = Utility::FindChildWidgetWithText(MenuItemBlock);
+		if (!ChildText.IsEmpty())
 		{
-			SlateQueryContext.StructuredContext.WidgetPath = TArray<FString>();
-			// Add simplified widget path
-			for (int32 WidgetIndex = 0; WidgetIndex < InWidgetPath.Widgets.Num(); WidgetIndex++)
-			{
-				const FArrangedWidget* ArrangedWidget = &InWidgetPath.Widgets[WidgetIndex];
-				const TSharedRef<SWidget>& ThisWidget = ArrangedWidget->Widget;
-				FString NextWidgetString = ThisWidget->GetType().ToString();
-				// Skip uninteresting widgets
-				if (NextWidgetString != TEXT("SBorder") && NextWidgetString != TEXT("SBox") &&
-					NextWidgetString != TEXT("SScrollBorder") &&
-					NextWidgetString != TEXT("SOverlay") && NextWidgetString != TEXT("SVerticalBox") &&
-					NextWidgetString != TEXT("SHorizontalBox") && NextWidgetString != TEXT("SSplitter") &&
-					NextWidgetString != TEXT("SDockingSplitter"))
-				{
-					SlateQueryContext.StructuredContext.WidgetPath->Add(NextWidgetString);
-				}
-			}
+			ItemName = ChildText;
+			ItemDescriptor = LOCTEXT("ItemDescriptor_Menu", "menu item");
+			LastPickedWidget = MenuItemBlock;
+			bIsUIWidget = true;
+			return true;
 		}
 	}
+	return false;
+}
 
-	// OPTIONAL - We're done using these members. We can clear them to reduce size as they're
-	// not used below.
-	SlateQueryContext.LastPickedWidget.Reset();
+bool FSlateQuerier::ItemIsButton()
+{
+	static const TSet<FName> ToolbarButtonTypes = TSet<FName>({
+		"SToolBarButtonBlock", "SToolBarComboButtonBlock"
+		});
 
-	// Send widget query to AI Assistant.
-	TSharedPtr<SAIAssistantWebBrowser> WebBrowser =
-		UAIAssistantSubsystem::GetAIAssistantWebBrowserWidget();
-	WebBrowser->CreateConversation();
+	static const TSet<FName> SimpleButtonTypes = TSet<FName>({
+			"SButton",
+			"SPrimaryButton",
+			"SCheckBox"
+		});
 
-	const FString VisiblePromptString =
-		CleanExtraWhiteSpaceFromString(SlateQueryContext.GeneratedQuery.ToString());
+	return (ItemIsButtonType(ToolbarButtonTypes, LOCTEXT("ItemDescriptor_ToolbarButton", "toolbar button")) ||
+			ItemIsButtonType(SimpleButtonTypes, LOCTEXT("ItemDescriptor_Button", "button")));
+}
 
-	const FString HiddenContextString = 
-		CleanExtraWhiteSpaceFromString(SlateQueryContext.GeneratedQueryInstructions.ToString()) +
-		SlateQueryContext.LanguageInstructions.ToString() +
-		TEXT(" ContextData=") +
-		SlateQueryContext.StructuredContext.ToJson(false);
+bool FSlateQuerier::ItemIsButtonType(const TSet<FName> ButtonWidgetTypes, FText InItemDescriptor)
+{
+	TSharedPtr<SWidget> Button = Utility::FindClosestWidgetPtrOfTypes(WidgetPath, ButtonWidgetTypes);
 
-	WebBrowser->AddUserMessageToConversation(VisiblePromptString, HiddenContextString);
+	// Disabled buttons *contain* the menu entry block.
+	if (!Button.IsValid() && WidgetPath.GetLastWidget()->GetChildren()->Num() > 0)
+	{
+		for (int32 ChildIdx = 0; ChildIdx < WidgetPath.GetLastWidget()->GetChildren()->Num(); ChildIdx++)
+		{
+			TSharedRef<SWidget> ThisWidget = WidgetPath.GetLastWidget()->GetChildren()->GetChildAt(ChildIdx);
+			if (ButtonWidgetTypes.Contains(ThisWidget->GetType()))
+			{
+				Button = ThisWidget.ToSharedPtr();
+			}
+		}
+	}
+
+	return Button.IsValid() && SetItemAsButton(Button, InItemDescriptor);
 }
 
+bool FSlateQuerier::SetItemAsButton(const TSharedPtr<SWidget> Button, FText InItemDescriptor)
+{
+	FText ChildText = Utility::FindChildWidgetWithText(Button);
+	if (!ChildText.IsEmpty() || !CurrentToolTipText.IsEmpty())
+	{
+		if (!ChildText.IsEmpty())
+		{
+			ItemName = ChildText;
+		}
+		else
+		{
+			// The tool tip can be used in this case.
+			ItemName = CurrentToolTipText;
+		}
+		ItemDescriptor = InItemDescriptor;
+		LastPickedWidget = Button;
+		bIsUIWidget = true;
+		return true;
+	}
+	return false;
+}
 
 #undef LOCTEXT_NAMESPACE
diff --git a/Engine/Plugins/Experimental/AIAssistant/Source/AIAssistant/Private/AIAssistantSlateQuerier.h b/Engine/Plugins/Experimental/AIAssistant/Source/AIAssistant/Private/AIAssistantSlateQuerier.h
--- a/Engine/Plugins/Experimental/AIAssistant/Source/AIAssistant/Private/AIAssistantSlateQuerier.h
+++ b/Engine/Plugins/Experimental/AIAssistant/Source/AIAssistant/Private/AIAssistantSlateQuerier.h
@@ -5,7 +5,8 @@
 
 
 #include "Layout/WidgetPath.h"
-
+#include "Serialization/JsonSerializable.h"
+#include "Serialization/JsonSerializerMacros.h"
 
 //
 // UE::AIAssistant::SlateQuerier
@@ -23,4 +24,213 @@ namespace UE::AIAssistant::SlateQuerier
 	 * Initiates an AI Assistant query to describe a Slate widget.
 	 */
 	void QueryAIAssistantAboutSlateWidget(const FWidgetPath& InWidgetPath);
+
+	/**
+	 * Store structured context for query to pass along to AI Assistant as JSON.
+	 */
+	struct FSlateQueryStructuredContext : public FJsonSerializable
+	{
+		FString TopWindowName;
+		FString UnrealEditorVersion;
+		FString UnrealEditorLanguage;
+		FString UnrealEditorLocale;
+		TOptional<FString> WindowName;
+		TOptional<FString> EditorName;
+		TOptional<FString> EditorMode;
+		TOptional<FString> ActiveToolName;
+		TOptional<FString> TabName;
+		TOptional<FString> ToolTipText;
+		TOptional<FString> TextUnderCursor;
+		TOptional<FString> WidgetType;
+		TOptional<FString> ActorClass;
+		TOptional<FString> ItemName;
+		TOptional<FString> ItemDescriptor;
+		TOptional<FString> CategoryName;
+		TOptional<TArray<FString>> WidgetPath;
+
+		BEGIN_JSON_SERIALIZER
+		JSON_SERIALIZE("toplevel_window_name", TopWindowName);
+		JSON_SERIALIZE("unreal_editor_version", UnrealEditorVersion);
+		JSON_SERIALIZE("unreal_editor_language", UnrealEditorLanguage);
+		JSON_SERIALIZE("unreal_editor_locale", UnrealEditorLocale);
+		JSON_SERIALIZE_OPTIONAL("window_name", WindowName);
+		JSON_SERIALIZE_OPTIONAL("editor_name", EditorName);
+		JSON_SERIALIZE_OPTIONAL("editor_mode", EditorMode);
+		JSON_SERIALIZE_OPTIONAL("active_tool", ActiveToolName);
+		JSON_SERIALIZE_OPTIONAL("tab_name", TabName);
+		JSON_SERIALIZE_OPTIONAL("tooltip_text", ToolTipText);
+		JSON_SERIALIZE_OPTIONAL("text_under_cursor", TextUnderCursor);
+		JSON_SERIALIZE_OPTIONAL("widget_type", WidgetType);
+		JSON_SERIALIZE_OPTIONAL("actor_class", ActorClass);
+		JSON_SERIALIZE_OPTIONAL("item_name", ItemName);
+		JSON_SERIALIZE_OPTIONAL("item_descriptor", ItemDescriptor);
+		JSON_SERIALIZE_OPTIONAL("category_name", CategoryName);
+		JSON_SERIALIZE_OPTIONAL_ARRAY("widget_path", WidgetPath);
+		END_JSON_SERIALIZER
+	};
+
+	/**
+	 * Implements Slate widget querying functionality.
+	 */
+	class FSlateQuerier
+	{
+	public:
+		FSlateQuerier(const FWidgetPath& InWidgetPath);
+
+		/**
+		 * Inspects widget path, determines query parameters, and sends query to Assistant backend.
+		 */
+		bool GenerateAndSendQuery();
+
+		/**
+		 * True if a query has been successfully constructed.
+		 */
+		bool IsValid() { return !GeneratedQuery.IsEmpty(); }
+
+	protected:
+		/**
+		 * Inspect widget path to determine query parameters.
+		 */
+		void IdentifyObjectOfQuery();
+
+		/**
+		 * Generate query and hidden instructions to Assistant.
+		 */
+		void GenerateQuery();
+
+		/**
+		 * Generate context packet JSON to send to Assistant for additonal context clues.
+		 */
+		void GenerateStructuredContext();
+
+		/**
+		 * Send query off to Assistant backend API.
+		 */
+		void SendQuery();
+
+		/**
+		 * Determine editor or window info from the widget path for purposes of querying the Assistant.
+		 * Return true if successful.
+		 */
+		bool FindEditorName();
+
+		/**
+		 * Determine tab or panel info from the widget path for purposes of querying the Assistant.
+		 * Return true if successful.
+		 */
+		bool FindTabName();
+
+		/**
+		 * Return true and store information if widget path is in a level or asset editor.
+		 */
+		bool IsInEditor();
+
+		/**
+		 * Return true and store information if widget path is in a drawer overlay.
+		 */
+		bool IsInDrawer();
+
+		/**
+		 * Determine what type of widget the widget path is on for purposes of querying the Assistant.
+		 * Return true if successful.
+		 */
+		bool FindItemName();
+
+		/**
+		 * Return true and store information if widget path is in a viewport.
+		 */
+		bool ItemIsInViewport();
+
+		/**
+		 * Return true and store information if widget path is in a graph node.
+		 */
+		bool ItemIsGraphNode();
+
+		/**
+		 * Return true and store information if widget path is in a search box.
+		 */
+		bool ItemIsSearchBox();
+
+		/**
+		 * Return true and store information if widget path is in a console input box.
+		 */
+		bool ItemIsConsoleInputBox();
+
+		/**
+		 * Return true and store information if widget path is in a Details panel property.
+		 */
+		bool ItemIsDetailsProperty();
+
+		/**
+		 * Return true and store information if widget path is in a plugin in Plugins browser.
+		 */
+		bool ItemIsPlugin();
+
+		/**
+		 * Return true and store information if widget path is on a navigation breadcrumb trail button.
+		 */
+		bool ItemIsBreadcrumbTrailButton();
+
+		/**
+		 * Return true and store information if widget path is on an asset in Content Browser.
+		 */
+		bool ItemIsAsset();
+
+		/**
+		 * Return true and store information if widget path is on an Actor in the Outliner panel, or on an
+		 * instance or component in the SubobjectInstanceEditor.
+		 */
+		bool ItemIsOutlinerActorOrSubobject();
+
+		/**
+		 * Return true and store information if widget path is on a Details panel filter checkbox.
+		 */
+		bool ItemIsDetailPanelFilterCheckbox();
+
+		/**
+		 * Return true and store information if widget path is on a menu item.
+		 */
+		bool ItemIsMenuItem();
+
+		/**
+		 * Return true and store information if widget path is on a button.
+		 */
+		bool ItemIsButton();
+
+		/**
+		 * Helper method for ItemIsButton
+		 * Return true and store information if widget path is on a button.
+		 */
+		bool ItemIsButtonType(const TSet<FName> InButtonWidgetTypes, FText InItemDescriptor);
+
+		/**
+		 * Helper method for buttons.  Returns true if method was able to successfully infer a name for the
+		 * button.
+		 */
+		bool SetItemAsButton(const TSharedPtr<SWidget> Button, FText InItemDescriptor);
+
+
+	private:
+		FWidgetPath OriginalWidgetPath;
+		FWidgetPath WidgetPath;
+		FText CurrentToolTipText;
+		FText FormattedWindowName;
+		FText FormattedTabName;
+		FText TextUnderCursor;
+		FText ItemName;
+		FText ItemDescriptor;
+		FString SelectedActorName;
+		FString SelectedActorClass;
+		TSharedPtr<SWidget> LastPickedWidget;
+
+		bool bIsUIWidget = false;
+		bool bIsObject = false;
+		bool bIsInAssistantPanel = false;
+
+		FText GeneratedQuery;
+		FText GeneratedQueryInstructions;
+		FText LanguageInstructions;
+
+		FSlateQueryStructuredContext StructuredContext;
+	};
 };
diff --git a/Engine/Plugins/Experimental/AIAssistant/Source/AIAssistant/Private/AIAssistantSlateQuerierUtils.cpp b/Engine/Plugins/Experimental/AIAssistant/Source/AIAssistant/Private/AIAssistantSlateQuerierUtils.cpp
new file mode 100644
--- /dev/null
+++ b/Engine/Plugins/Experimental/AIAssistant/Source/AIAssistant/Private/AIAssistantSlateQuerierUtils.cpp
@@ -0,0 +1,101 @@
+// Copyright Epic Games, Inc. All Rights Reserved.
+
+#include "AIAssistantSlateQuerierUtils.h"
+
+#include "Framework/Application/SlateApplication.h"
+#include "Internationalization/Regex.h"
+#include "Widgets/Text/STextBlock.h"
+#include "Widgets/Text/SRichTextBlock.h"
+
+
+using namespace UE::AIAssistant::SlateQuerier;
+
+FString Utility::CleanExtraWhiteSpaceFromString(const FString& InString)
+{
+	FString CleanString = "";
+	static const FRegexPattern WhiteSpacePattern(TEXT("(\\S+)"));
+	FRegexMatcher Matcher(WhiteSpacePattern, InString);
+	while (Matcher.FindNext())
+	{
+		if (!CleanString.IsEmpty())
+		{
+			CleanString += " ";
+		}
+		CleanString += Matcher.GetCaptureGroup(1);
+	}
+
+	return MoveTemp(CleanString);
+}
+
+FText Utility::GetWidgetText(const TSharedRef<SWidget> InWidget)
+{
+	FText OutText;
+	// refuse any that only contain numbers (they're probably just values)
+	const FRegexPattern AlphaPattern(TEXT("[A-Za-z]+"));
+	FText WidgetText;
+
+	if (InWidget->GetType() == "STextBlock")
+	{
+		TSharedRef<STextBlock> TextBlock = StaticCastSharedRef<STextBlock>(InWidget);
+		WidgetText = TextBlock->GetText();
+	}
+	else if (InWidget->GetType() == "SRichTextBlock")
+	{
+		TSharedRef<SRichTextBlock> TextBlock = StaticCastSharedRef<SRichTextBlock>(InWidget);
+		WidgetText = TextBlock->GetText();
+	}
+	if (!WidgetText.IsEmpty())
+	{
+		FRegexMatcher AlphaMatcher(AlphaPattern, WidgetText.ToString());
+		if (AlphaMatcher.FindNext())
+		{
+			OutText = WidgetText;
+		}
+	}
+
+	return OutText;
+}
+
+FText Utility::FindTextUnderCursor(const FWidgetPath& InWidgetPath)
+{
+	return InWidgetPath.IsValid() ? Utility::GetWidgetText(InWidgetPath.GetLastWidget()) : FText();
+}
+
+FText Utility::FindChildWidgetWithText(const TSharedPtr<SWidget> InWidget)
+{
+	if (InWidget.IsValid())
+	{
+		FText WidgetText = Utility::GetWidgetText(InWidget.ToSharedRef());
+		if (!WidgetText.IsEmpty())
+		{
+			return WidgetText;
+		}
+
+		// get all children and see if any are text widgets.
+		for (int32 ChildIdx = 0; ChildIdx < InWidget->GetChildren()->Num(); ChildIdx++)
+		{
+			TSharedRef<SWidget> ThisWidget = InWidget->GetChildren()->GetChildAt(ChildIdx);
+			const FText ChildText = Utility::FindChildWidgetWithText(ThisWidget.ToSharedPtr());
+			if (!ChildText.IsEmpty())
+			{
+				return ChildText;
+			}
+		}
+	}
+
+	return FText();
+}
+
+FText Utility::FindCurrentToolTipText()
+{
+	for (const TArray<TSharedRef<SWindow>> AllWindows = FSlateApplication::Get().GetTopLevelWindows();
+		const TSharedRef<SWindow>& ThisWindow : AllWindows)
+	{
+		if (ThisWindow->GetType() == EWindowType::ToolTip && ThisWindow->IsVisible())
+		{
+			return Utility::FindChildWidgetWithText(ThisWindow);
+		}
+	}
+
+	return FText();
+}
diff --git a/Engine/Plugins/Experimental/AIAssistant/Source/AIAssistant/Private/AIAssistantSlateQuerierUtils.h b/Engine/Plugins/Experimental/AIAssistant/Source/AIAssistant/Private/AIAssistantSlateQuerierUtils.h
new file mode 100644
--- /dev/null
+++ b/Engine/Plugins/Experimental/AIAssistant/Source/AIAssistant/Private/AIAssistantSlateQuerierUtils.h
@@ -0,0 +1,37 @@
+// Copyright Epic Games, Inc. All Rights Reserved.
+
+#pragma once
+
+#include "Containers/UnrealString.h"
+#include "Layout/WidgetPath.h"
+#include "Templates/SharedPointer.h"
+
+
+namespace UE::AIAssistant::SlateQuerier::Utility
+{
+	/**
+	 * Return copy of InString cleaned of redundant whitespace.
+	 */
+	FString CleanExtraWhiteSpaceFromString(const FString& InString);
+
+	/**
+	 * Look for text in widget under cursor and return it.  Return empty FText if not found.
+	 */
+	FText FindTextUnderCursor(const FWidgetPath& InWidgetPath);
+
+	/**
+	 * Return text from InWidget if it is a text item.  Return empty FText if not found.
+	 */
+	FText GetWidgetText(const TSharedRef<SWidget> InWidget);
+
+	/**
+	 * Search children of InWidget and return text from first one found with alphabetic text.
+	 * Return empty FText if not found.
+	 */
+	FText FindChildWidgetWithText(const TSharedPtr<SWidget> InWidget);
+
+	/**
+	 * Return first active tooltip text found.  Return empty FText if not found.
+	 */
+	FText FindCurrentToolTipText();
+};
diff --git a/Engine/Plugins/Experimental/AIAssistant/Source/AIAssistant/Private/AIAssistantSlateQuerierUtilsTypeSearch.cpp b/Engine/Plugins/Experimental/AIAssistant/Source/AIAssistant/Private/AIAssistantSlateQuerierUtilsTypeSearch.cpp
new file mode 100644
--- /dev/null
+++ b/Engine/Plugins/Experimental/AIAssistant/Source/AIAssistant/Private/AIAssistantSlateQuerierUtilsTypeSearch.cpp
@@ -0,0 +1,133 @@
+// Copyright Epic Games, Inc. All Rights Reserved.
+
+#include "AIAssistantSlateQuerierUtilsTypeSearch.h"
+
+using namespace UE::AIAssistant::SlateQuerier;
+
+TSharedPtr<SWidget> Utility::FindClosestWidgetPtrOfTypes(const FWidgetPath& InWidgetPath, const TSet<FName> InWidgetTypes)
+{
+	TSharedPtr<SWidget> ClosestWidget;
+	for (int32 WidgetIndex = InWidgetPath.Widgets.Num() - 1; WidgetIndex >= 0; --WidgetIndex)
+	{
+		const FArrangedWidget* ArrangedWidget = &InWidgetPath.Widgets[WidgetIndex];
+		const TSharedRef<SWidget>& ThisWidget = ArrangedWidget->Widget;
+		if (InWidgetTypes.Contains(ThisWidget->GetType()))
+		{
+			ClosestWidget = ThisWidget.ToSharedPtr();
+			break;
+		}
+	}
+	return ClosestWidget;
+}
+
+TSharedPtr<SWidget> Utility::FindClosestWidgetPtrOfType(const FWidgetPath& InWidgetPath, const FName& InWidgetType)
+{
+	const TSet<FName> WidgetTypes = TSet<FName>({InWidgetType});
+	return Utility::FindClosestWidgetPtrOfTypes(InWidgetPath, WidgetTypes);
+}
+
+TSharedPtr<SWidget> Utility::FindClosestWidgetPtrOfSortedTypes(
+	const FWidgetPath& InWidgetPath, const TConstArrayView<FName>& InWidgetTypes)
+{
+	for (FName WidgetType : InWidgetTypes)
+	{
+		TSharedPtr<SWidget> ClosestWidget = Utility::FindClosestWidgetPtrOfType(InWidgetPath, WidgetType);
+		if (ClosestWidget.IsValid())
+		{
+			return ClosestWidget;
+		}
+	}
+
+	return nullptr;
+}
+
+TSharedPtr<SWidget> Utility::FindFirstWidgetPtrOfType(const FWidgetPath& InWidgetPath, const FName& InWidgetType)
+{
+	TSharedPtr<SWidget> FirstWidget;
+	for (int32 WidgetIndex = 0; WidgetIndex < InWidgetPath.Widgets.Num(); WidgetIndex++)
+	{
+		const FArrangedWidget* ArrangedWidget = &InWidgetPath.Widgets[WidgetIndex];
+		const TSharedRef<SWidget>& ThisWidget = ArrangedWidget->Widget;
+		if (ThisWidget->GetType() == InWidgetType)
+		{
+			FirstWidget = ThisWidget.ToSharedPtr();
+			break;
+		}
+	}
+	return FirstWidget;
+}
+
+void Utility::FindChildWidgetsOfType(TArray<TSharedRef<SWidget>>& OutWidgets, TSharedPtr<SWidget> InWidget, const FName& InWidgetType)
+{
+	if (InWidget.IsValid())
+	{
+		for (int32 ChildIdx = 0; ChildIdx < InWidget->GetChildren()->Num(); ChildIdx++)
+		{
+			TSharedRef<SWidget> ThisWidget = InWidget->GetChildren()->GetChildAt(ChildIdx);
+			if (ThisWidget->GetType() == InWidgetType)
+			{
+				OutWidgets.Add(ThisWidget);
+			}
+			Utility::FindChildWidgetsOfType(OutWidgets, ThisWidget.ToSharedPtr(), InWidgetType);
+		}
+	}
+}
+
+
+
+TSharedPtr<SWidget> Utility::FindChildWidgetOfType(const TSharedPtr<SWidget> InWidget, const FName& InWidgetType)
+{
+	if (InWidget.IsValid())
+	{
+		for (int32 ChildIdx = 0; ChildIdx < InWidget->GetChildren()->Num(); ChildIdx++)
+		{
+			TSharedRef<SWidget> ThisWidget = InWidget->GetChildren()->GetChildAt(ChildIdx);
+			if (ThisWidget->GetType() == InWidgetType)
+			{
+				return ThisWidget.ToSharedPtr();
+			}
+			else
+			{
+				TSharedPtr<SWidget> ChildWidget = Utility::FindChildWidgetOfType(ThisWidget.ToSharedPtr(), InWidgetType);
+				if (ChildWidget.IsValid() && ChildWidget->GetType() == InWidgetType)
+				{
+					return ChildWidget;
+				}
+			}
+		}
+	}
+	return TSharedPtr<SWidget>();
+}
+
+TSharedPtr<SWidget> Utility::FindClosestMenuItem(const FWidgetPath& WidgetPath)
+{
+	TSharedPtr<SWidget> MenuItemBlock;
+	if (WidgetPath.IsValid())
+	{
+		// is it a menu item? Look up for an SMenuEntryBlock or SWidgetBlock.
+		static const TStaticArray<FName, 2> MenuTypes = {"SMenuEntryBlock", "SWidgetBlock"};
+		MenuItemBlock = Utility::FindClosestWidgetPtrOfSortedTypes(WidgetPath, MenuTypes);
+
+		// disabled menu items *contain* the menu entry block
+		if (!MenuItemBlock.IsValid())
+		{
+			if (WidgetPath.GetLastWidget()->GetChildren()->Num() > 0)
+			{
+				for (int32 ChildIdx = 0; ChildIdx < WidgetPath.GetLastWidget()->GetChildren()->Num(); ChildIdx++)
+				{
+					TSharedRef<SWidget> ThisWidget = WidgetPath.GetLastWidget()->GetChildren()->GetChildAt(ChildIdx);
+					for (FName MenuType : MenuTypes)
+					{
+						if (ThisWidget->GetType() == MenuType)
+						{
+							MenuItemBlock = ThisWidget.ToSharedPtr();
+							break;
+						}
+					}
+				}
+			}
+		}
+	}
+
+	return MenuItemBlock;
+}
diff --git a/Engine/Plugins/Experimental/AIAssistant/Source/AIAssistant/Private/AIAssistantSlateQuerierUtilsTypeSearch.h b/Engine/Plugins/Experimental/AIAssistant/Source/AIAssistant/Private/AIAssistantSlateQuerierUtilsTypeSearch.h
new file mode 100644
--- /dev/null
+++ b/Engine/Plugins/Experimental/AIAssistant/Source/AIAssistant/Private/AIAssistantSlateQuerierUtilsTypeSearch.h
@@ -0,0 +1,52 @@
+// Copyright Epic Games, Inc. All Rights Reserved.
+
+#pragma once
+
+#include "Layout/WidgetPath.h"
+
+
+namespace UE::AIAssistant::SlateQuerier::Utility
+{
+	/**
+	 * Search backwards from end of InWidgetPath and return first widget that matches any type in WidgetTypes.
+	 * Return null ptr if not found.
+	 */
+	TSharedPtr<SWidget> FindClosestWidgetPtrOfTypes(const FWidgetPath& InWidgetPath, const TSet<FName> InWidgetTypes);
+
+	/**
+	 * Search backwards from end of InWidgetPath and return first widget that matches WidgetType.
+	 * Return null ptr if not found.
+	 */
+	TSharedPtr<SWidget> FindClosestWidgetPtrOfType(const FWidgetPath& InWidgetPath, const FName& InWidgetType);
+
+	/**
+	 * Search backwards from end of InWidgetPath and return first widget that matches a type in WidgetTypes,
+	 * testing each type in order.
+	 * Return null ptr if not found.
+	 */
+	TSharedPtr<SWidget> FindClosestWidgetPtrOfSortedTypes(
+		const FWidgetPath& InWidgetPath, const TConstArrayView<FName>& InWidgetTypes);
+
+	/**
+	 * Search forward from start of InWidgetPath and return first widget that matches WidgetType.
+	 * Return null ptr if not found.
+	 */
+	TSharedPtr<SWidget> FindFirstWidgetPtrOfType(const FWidgetPath& InWidgetPath, const FName& InWidgetType);
+
+	/**
+	 * Search children of InWidget for any widgets that match InWidgetType; return all found widgets in OutWidgets.
+	 */
+	void FindChildWidgetsOfType(TArray<TSharedRef<SWidget>>& OutWidgets, TSharedPtr<SWidget> InWidget, const FName& InWidgetType);
+
+	/**
+	 * Search children of InWidget and return first widget that matches InWidgetType.
+	 * Return null ptr if not found.
+	 */
+	TSharedPtr<SWidget> FindChildWidgetOfType(const TSharedPtr<SWidget> InWidget, const FName& InWidgetType);
+
+	/**
+	 * Search backwards from end of InWidgetPath and return first menu item type widget found.
+	 * Return null ptr if not found.
+	 */
+	TSharedPtr<SWidget> FindClosestMenuItem(const FWidgetPath& InWidgetPath);
+};
diff --git a/Engine/Plugins/Experimental/AIAssistant/Source/AIAssistant/Private/Tests/AIAssistantSlateQuerierTest.cpp b/Engine/Plugins/Experimental/AIAssistant/Source/AIAssistant/Private/Tests/AIAssistantSlateQuerierTest.cpp
new file mode 100644
--- /dev/null
+++ b/Engine/Plugins/Experimental/AIAssistant/Source/AIAssistant/Private/Tests/AIAssistantSlateQuerierTest.cpp
@@ -0,0 +1,99 @@
+// Copyright Epic Games, Inc. All Rights Reserved.
+
+#include "Misc/AutomationTest.h"
+
+#include "AIAssistantSlateQuerierTestHarness.h"
+
+#if WITH_DEV_AUTOMATION_TESTS
+
+// NOTE: Currently does not fully cover all the SlateQuerier functionality, as that will be refactored heavily
+// in the near future to be more granular.
+
+using namespace UE::AIAssistant::SlateQuerier;
+
+class SlateQuerierTestWidget : public STestWidget
+{
+  public:
+	void Construct(const FArguments& InArgs)
+	{
+		// TODO: add more complex constructions for more complex testing.
+		ChildSlot
+			[
+				SNew(SVerticalBox)
+				+ SVerticalBox::Slot()
+				.AutoHeight()
+				[
+					SAssignNew(OuterHBox, SHorizontalBox)
+					+ SHorizontalBox::Slot()
+					.AutoWidth()
+					[
+						SAssignNew(InnerHBox, SHorizontalBox)
+						+ SHorizontalBox::Slot()
+						.AutoWidth()
+						[
+							SNew(SHorizontalBox)
+							+ SHorizontalBox::Slot()
+							.AutoWidth()
+							[
+								SAssignNew(TextBlock, STextBlock)
+								.Text(FText::FromString(TEXT("Hello, Slate!")))
+							 ]
+						 ]
+						+ SHorizontalBox::Slot()
+						.AutoWidth()
+						[
+							SAssignNew(CheckBoxButton, SCheckBox)
+							.Type(ESlateCheckBoxType::ToggleButton)
+							[
+								SNew(STextBlock)
+								.Text(FText::FromString(TEXT("CheckBox Name")))
+							 ]
+						 ]
+						+ SHorizontalBox::Slot()
+						.AutoWidth()
+						[
+							SAssignNew(SimpleButton, SButton)
+							.Text(FText::FromString(TEXT("Button Name")))
+						 ]
+					 ]
+				 ]
+			 ];
+	}
+};
+
+
+IMPLEMENT_SIMPLE_AUTOMATION_TEST(
+	FAIAssistantSlateQuerierTestCreateWidgetPath,
+	"AI.Assistant.SlateQuerier.CreateWidgetPath",
+	SlateQuerierTestFlags);
+
+bool FAIAssistantSlateQuerierTestCreateWidgetPath::RunTest(const FString& UnusedParameters)
+{
+	FSlateTestHarness TestHarness(SNew(SlateQuerierTestWidget));
+
+	TestHarness.GenerateWidgetPathTo(TestHarness.TestWidget);
+	TestTrue(TEXT("WidgetPath is valid"), TestHarness.WidgetPath.IsValid());
+
+	return true;
+}
+
+IMPLEMENT_SIMPLE_AUTOMATION_TEST(
+	FAIAssistantSlateQuerierTestItemIsButton,
+	"AI.Assistant.SlateQuerier.ItemIsButton",
+	SlateQuerierTestFlags);
+
+bool FAIAssistantSlateQuerierTestItemIsButton::RunTest(const FString& UnusedParameters)
+{
+	FSlateTestHarness TestHarness(SNew(SlateQuerierTestWidget));
+
+	TestHarness.GenerateWidgetPathTo(TestHarness.TestWidget->SimpleButton);
+	TestTrue(TEXT("Identify SimpleButton"), TestHarness.SlateQuerier->TestItemIsButton());
+	TestHarness.GenerateWidgetPathTo(TestHarness.TestWidget->CheckBoxButton);
+	TestTrue(TEXT("Identify CheckBox"), TestHarness.SlateQuerier->TestItemIsButton());
+	TestHarness.GenerateWidgetPathTo(TestHarness.TestWidget->TextBlock);
+	TestFalse(TEXT("Text is not Button"), TestHarness.SlateQuerier->TestItemIsButton());
+
+	return true;
+}
+
+#endif	// WITH_DEV_AUTOMATION_TESTS
diff --git a/Engine/Plugins/Experimental/AIAssistant/Source/AIAssistant/Private/Tests/AIAssistantSlateQuerierTestHarness.cpp b/Engine/Plugins/Experimental/AIAssistant/Source/AIAssistant/Private/Tests/AIAssistantSlateQuerierTestHarness.cpp
new file mode 100644
--- /dev/null
+++ b/Engine/Plugins/Experimental/AIAssistant/Source/AIAssistant/Private/Tests/AIAssistantSlateQuerierTestHarness.cpp
@@ -0,0 +1,96 @@
+// Copyright Epic Games, Inc. All Rights Reserved.
+
+#include "AIAssistantSlateQuerierTestHarness.h"
+
+#include "Widgets/SCompoundWidget.h"
+#include "Widgets/Input/SButton.h"
+#include "Widgets/Input/SCheckBox.h"
+#include "Widgets/Text/STextBlock.h"
+#include "Widgets/Text/SRichTextBlock.h"
+#include "Framework/Application/SlateApplication.h"
+
+#include "AIAssistantLog.h"
+
+UE::AIAssistant::SlateQuerier::FSlateTestHarness::FSlateTestHarness(TSharedRef<STestWidget> InWidget)
+{
+	TSharedRef<SWindow> Window = SNew(SWindow)
+		.Title(FText::FromString(TEXT("Test Slate Window")))
+		.ClientSize(FVector2D(400, 200));
+
+	FSlateApplication::Get().AddWindow(Window);
+	Window->SetContent(InWidget);
+	TestWidget = InWidget.ToSharedPtr();
+	ThisWindow = Window.ToSharedPtr();
+}
+
+UE::AIAssistant::SlateQuerier::FSlateTestHarness::~FSlateTestHarness()
+{
+	if (ThisWindow)
+	{
+		FSlateApplication::Get().RequestDestroyWindow(ThisWindow.ToSharedRef());
+	}
+}
+
+void UE::AIAssistant::SlateQuerier::FSlateTestHarness::GenerateWidgetPathTo(const TSharedPtr<SWidget> InWidget)
+{
+	FSlateApplication::Get().GeneratePathToWidgetUnchecked(InWidget.ToSharedRef(), WidgetPath);
+	FTestSlateQuerier* Querier = new FTestSlateQuerier(WidgetPath);
+	SlateQuerier = TSharedPtr<FTestSlateQuerier>(Querier);
+}
+
+void UE::AIAssistant::SlateQuerier::FSlateTestHarness::LogWidgetPath()
+{
+	for (int32 WidgetIndex = 0; WidgetIndex < WidgetPath.Widgets.Num(); WidgetIndex++)
+	{
+		const FArrangedWidget* ArrangedWidget = &WidgetPath.Widgets[WidgetIndex];
+		TSharedPtr<SWidget> ThisWidget = ArrangedWidget->Widget.ToSharedPtr();
+		FName WidgetTypeName = ThisWidget->GetType();
+		UE_LOG(LogAIAssistant, Display, TEXT("Widget %d: %s: %p"),
+			   WidgetIndex, *(WidgetTypeName.ToString()), ThisWidget.Get());
+	}
+}
+
+void UE::AIAssistant::SlateQuerier::STestWidget::Construct(const FArguments& InArgs)
+{
+	ChildSlot
+		[
+			SNew(SVerticalBox)
+			+ SVerticalBox::Slot()
+			.AutoHeight()
+			[
+				SAssignNew(OuterHBox, SHorizontalBox)
+				+ SHorizontalBox::Slot()
+				.AutoWidth()
+				[
+					SAssignNew(InnerHBox, SHorizontalBox)
+					+ SHorizontalBox::Slot()
+					.AutoWidth()
+					[
+						SNew(SHorizontalBox)
+						+ SHorizontalBox::Slot()
+						.AutoWidth()
+						[
+							SAssignNew(TextBlock, STextBlock)
+							.Text(FText::FromString(TEXT("Hello, Slate!")))
+						 ]
+					 ]
+					+ SHorizontalBox::Slot()
+					.AutoWidth()
+					[
+						SAssignNew(CheckBoxButton, SCheckBox)
+						.Type(ESlateCheckBoxType::ToggleButton)
+						[
+							SNew(STextBlock)
+							.Text(FText::FromString(TEXT("CheckBox Name")))
+						 ]
+					 ]
+					+ SHorizontalBox::Slot()
+					.AutoWidth()
+					[
+						SAssignNew(SimpleButton, SButton)
+						.Text(FText::FromString(TEXT("Button Name")))
+					 ]
+				 ]
+			 ]
+		 ];
+}
diff --git a/Engine/Plugins/Experimental/AIAssistant/Source/AIAssistant/Private/Tests/AIAssistantSlateQuerierTestHarness.h b/Engine/Plugins/Experimental/AIAssistant/Source/AIAssistant/Private/Tests/AIAssistantSlateQuerierTestHarness.h
new file mode 100644
--- /dev/null
+++ b/Engine/Plugins/Experimental/AIAssistant/Source/AIAssistant/Private/Tests/AIAssistantSlateQuerierTestHarness.h
@@ -0,0 +1,52 @@
+// Copyright Epic Games, Inc. All Rights Reserved.
+
+#pragma once
+
+#include "Widgets/SCompoundWidget.h"
+#include "Widgets/Input/SButton.h"
+#include "Widgets/Input/SCheckBox.h"
+#include "Widgets/Text/STextBlock.h"
+
+#include "AIAssistantSlateQuerier.h"
+#include "AIAssistantTestFlags.h"
+
+
+namespace UE::AIAssistant::SlateQuerier
+{
+	const auto SlateQuerierTestFlags = AIAssistantTest::Flags | EAutomationTestFlags::NonNullRHI;
+
+	class STestWidget : public SCompoundWidget
+	{
+	  public:
+		SLATE_BEGIN_ARGS(STestWidget) {}
+		SLATE_END_ARGS()
+
+		void Construct(const FArguments& InArgs);
+
+	  public:
+		TSharedPtr<SHorizontalBox> OuterHBox;
+		TSharedPtr<SHorizontalBox> InnerHBox;
+		TSharedPtr<SButton> SimpleButton;
+		TSharedPtr<SCheckBox> CheckBoxButton;
+		TSharedPtr<STextBlock> TextBlock;
+	};
+
+	class FTestSlateQuerier : public FSlateQuerier
+	{
+	  public:
+		bool TestItemIsButton() {return ItemIsButton();}
+	};
+
+	struct FSlateTestHarness
+	{
+		TSharedPtr<SWindow> ThisWindow;
+		TSharedPtr<STestWidget> TestWidget;
+		TSharedPtr<FTestSlateQuerier> SlateQuerier;
+		FWidgetPath WidgetPath;
+
+		FSlateTestHarness(TSharedRef<STestWidget> InWidget);
+		~FSlateTestHarness();
+		void GenerateWidgetPathTo(const TSharedPtr<SWidget> InWidget);
+		void LogWidgetPath();
+	};
+};
diff --git a/Engine/Plugins/Experimental/AIAssistant/Source/AIAssistant/Private/Tests/AIAssistantSlateQuerierUtilsTest.cpp b/Engine/Plugins/Experimental/AIAssistant/Source/AIAssistant/Private/Tests/AIAssistantSlateQuerierUtilsTest.cpp
new file mode 100644
--- /dev/null
+++ b/Engine/Plugins/Experimental/AIAssistant/Source/AIAssistant/Private/Tests/AIAssistantSlateQuerierUtilsTest.cpp
@@ -0,0 +1,45 @@
+// Copyright Epic Games, Inc. All Rights Reserved.
+
+#include "Misc/AutomationTest.h"
+
+#include "AIAssistantSlateQuerierTestHarness.h"
+
+#include "AIAssistantSlateQuerierUtils.h"
+
+#if WITH_DEV_AUTOMATION_TESTS
+
+using namespace UE::AIAssistant::SlateQuerier;
+
+IMPLEMENT_SIMPLE_AUTOMATION_TEST(
+	FAIAssistantSlateQuerierTestFindChildText,
+	"AI.Assistant.SlateQuerier.Utils.FindChildText",
+	SlateQuerierTestFlags);
+
+bool FAIAssistantSlateQuerierTestFindChildText::RunTest(const FString& UnusedParameters)
+{
+	FSlateTestHarness TestHarness(SNew(STestWidget));
+
+	TestHarness.GenerateWidgetPathTo(TestHarness.TestWidget);
+	FText FoundText = Utility::FindChildWidgetWithText(TestHarness.TestWidget);
+	TestFalse(TEXT("Find non-empty child text in widget"), FoundText.IsEmpty());
+
+	return true;
+}
+
+IMPLEMENT_SIMPLE_AUTOMATION_TEST(
+	FAIAssistantSlateQuerierTestItemIsText,
+	"AI.Assistant.SlateQuerier.Utils.ItemIsText",
+	SlateQuerierTestFlags);
+
+bool FAIAssistantSlateQuerierTestItemIsText::RunTest(const FString& UnusedParameters)
+{
+	FSlateTestHarness TestHarness(SNew(STestWidget));
+
+	TestHarness.GenerateWidgetPathTo(TestHarness.TestWidget->TextBlock);
+	FText FoundText = Utility::FindTextUnderCursor(TestHarness.WidgetPath);
+	TestFalse(TEXT("Identify text block under cursor"), FoundText.IsEmpty());
+
+	return true;
+}
+
+#endif	// WITH_DEV_AUTOMATION_TESTS
diff --git a/Engine/Plugins/Experimental/AIAssistant/Source/AIAssistant/Private/Tests/AIAssistantSlateQuerierUtilsTypeSearchTest.cpp b/Engine/Plugins/Experimental/AIAssistant/Source/AIAssistant/Private/Tests/AIAssistantSlateQuerierUtilsTypeSearchTest.cpp
new file mode 100644
--- /dev/null
+++ b/Engine/Plugins/Experimental/AIAssistant/Source/AIAssistant/Private/Tests/AIAssistantSlateQuerierUtilsTypeSearchTest.cpp
@@ -0,0 +1,104 @@
+// Copyright Epic Games, Inc. All Rights Reserved.
+
+#include "Misc/AutomationTest.h"
+
+#include "AIAssistantLog.h"
+#include "AIAssistantSlateQuerierTestHarness.h"
+#include "AIAssistantSlateQuerierUtilsTypeSearch.h"
+
+#if WITH_DEV_AUTOMATION_TESTS
+
+// NOTE: Currently does not fully cover all the type-searching Utility functionality, that will be coming
+// soon, along with more test UI harness infrastructure for testing more complex Slate UI such as Menus and
+// Toolbars.
+
+using namespace UE::AIAssistant::SlateQuerier;
+
+IMPLEMENT_SIMPLE_AUTOMATION_TEST(
+	FAIAssistantSlateQuerierTestFindChildWidget,
+	"AI.Assistant.SlateQuerier.Utils.TypeSearch.FindChildWidget",
+	SlateQuerierTestFlags);
+
+bool FAIAssistantSlateQuerierTestFindChildWidget::RunTest(const FString& UnusedParameters)
+{
+	FSlateTestHarness TestHarness(SNew(STestWidget));
+
+	TSharedPtr<SWidget> ChildWidget = Utility::FindChildWidgetOfType(
+		TestHarness.TestWidget->SimpleButton, "STextBlock");
+	TestTrue(TEXT("Button's child text widget found"), ChildWidget.IsValid());
+	TestTrue(TEXT("Button's child text is text"), ChildWidget->GetType() == "STextBlock");
+
+	ChildWidget = Utility::FindChildWidgetOfType(TestHarness.TestWidget->TextBlock, "SButton");
+	TestFalse(TEXT("Text does not have child Button widget"), ChildWidget.IsValid());
+
+	return true;
+}
+
+IMPLEMENT_SIMPLE_AUTOMATION_TEST(
+	FAIAssistantSlateQuerierTestFindChildWidgets,
+	"AI.Assistant.SlateQuerier.Utils.TypeSearch.FindChildWidgets",
+	SlateQuerierTestFlags);
+
+bool FAIAssistantSlateQuerierTestFindChildWidgets::RunTest(const FString& UnusedParameters)
+{
+	FSlateTestHarness TestHarness(SNew(STestWidget));
+
+	TArray<TSharedRef<SWidget>> FoundWidgets;
+
+	Utility::FindChildWidgetsOfType(FoundWidgets, TestHarness.TestWidget, "STextBlock");
+	TestTrue(TEXT("Find multiple text widgets"), FoundWidgets.Num() > 1);
+
+	FoundWidgets.Empty();
+	Utility::FindChildWidgetsOfType(FoundWidgets, TestHarness.TestWidget, "SCheckBox");
+	TestTrue(TEXT("Find single checkbox widget"), FoundWidgets.Num() == 1);
+
+	FoundWidgets.Empty();
+	Utility::FindChildWidgetsOfType(FoundWidgets, TestHarness.TestWidget, "SNonExistent");
+	TestTrue(TEXT("Find zero nonexistent widget"), FoundWidgets.Num() == 0);
+
+	return true;
+}
+
+IMPLEMENT_SIMPLE_AUTOMATION_TEST(
+	FAIAssistantSlateQuerierTestFindClosestWidget,
+	"AI.Assistant.SlateQuerier.Utils.TypeSearch.FindClosestWidget",
+	SlateQuerierTestFlags);
+
+bool FAIAssistantSlateQuerierTestFindClosestWidget::RunTest(const FString& UnusedParameters)
+{
+	FSlateTestHarness TestHarness(SNew(STestWidget));
+
+	TestHarness.GenerateWidgetPathTo(TestHarness.TestWidget->CheckBoxButton);
+
+	TSharedPtr<SWidget> FoundWidget = Utility::FindClosestWidgetPtrOfType(
+		TestHarness.WidgetPath, "SHorizontalBox");
+	TestTrue(TEXT("Closest widget test finds inner widget"),
+			 FoundWidget == TestHarness.TestWidget->InnerHBox);
+	TestFalse(TEXT("Closest widget test does not find outer widget"),
+			 FoundWidget == TestHarness.TestWidget->OuterHBox);
+
+	return true;
+}
+
+IMPLEMENT_SIMPLE_AUTOMATION_TEST(
+	FAIAssistantSlateQuerierTestFindFirstWidget,
+	"AI.Assistant.SlateQuerier.Utils.TypeSearch.FindFirstWidget",
+	SlateQuerierTestFlags);
+
+bool FAIAssistantSlateQuerierTestFindFirstWidget::RunTest(const FString& UnusedParameters)
+{
+	FSlateTestHarness TestHarness(SNew(STestWidget));
+
+	TestHarness.GenerateWidgetPathTo(TestHarness.TestWidget->CheckBoxButton);
+
+	TSharedPtr<SWidget> FoundWidget = Utility::FindFirstWidgetPtrOfType(
+		TestHarness.WidgetPath, "SHorizontalBox");
+	TestTrue(TEXT("First widget test finds outer widget"),
+			 FoundWidget == TestHarness.TestWidget->OuterHBox);
+	TestFalse(TEXT("First widget test does not find inner widget"),
+			 FoundWidget == TestHarness.TestWidget->InnerHBox);
+
+	return true;
+}
+
+#endif	// WITH_DEV_AUTOMATION_TESTS
