[AI Assistant] Horde Issue 1117379:  Fix compile error in AIAssistantSlateQuerierTestHarness.h

#rnx

diff --git a/Engine/Plugins/Experimental/AIAssistant/Source/AIAssistant/Private/Tests/AIAssistantSlateQuerierTestHarness.h b/Engine/Plugins/Experimental/AIAssistant/Source/AIAssistant/Private/Tests/AIAssistantSlateQuerierTestHarness.h
--- a/Engine/Plugins/Experimental/AIAssistant/Source/AIAssistant/Private/Tests/AIAssistantSlateQuerierTestHarness.h
+++ b/Engine/Plugins/Experimental/AIAssistant/Source/AIAssistant/Private/Tests/AIAssistantSlateQuerierTestHarness.h
@@ -2,6 +2,8 @@
 
 #pragma once
 
+#if WITH_DEV_AUTOMATION_TESTS
+
 #include "Widgets/SCompoundWidget.h"
 #include "Widgets/Input/SButton.h"
 #include "Widgets/Input/SCheckBox.h"
@@ -50,3 +52,5 @@ namespace UE::AIAssistant::SlateQuerier
 		void LogWidgetPath();
 	};
 };
+
+#endif  // WITH_DEV_AUTOMATION_TESTS
