[AI Assistant] Fix Horde Issue 1117467: Compile errors in AIAssistantSlateQuerierTestHarness.cpp

#rnx

diff --git a/Engine/Plugins/Experimental/AIAssistant/Source/AIAssistant/Private/Tests/AIAssistantSlateQuerierTestHarness.cpp b/Engine/Plugins/Experimental/AIAssistant/Source/AIAssistant/Private/Tests/AIAssistantSlateQuerierTestHarness.cpp
--- a/Engine/Plugins/Experimental/AIAssistant/Source/AIAssistant/Private/Tests/AIAssistantSlateQuerierTestHarness.cpp
+++ b/Engine/Plugins/Experimental/AIAssistant/Source/AIAssistant/Private/Tests/AIAssistantSlateQuerierTestHarness.cpp
@@ -2,6 +2,8 @@
 
 #include "AIAssistantSlateQuerierTestHarness.h"
 
+#if WITH_DEV_AUTOMATION_TESTS
+
 #include "Widgets/SCompoundWidget.h"
 #include "Widgets/Input/SButton.h"
 #include "Widgets/Input/SCheckBox.h"
@@ -94,3 +96,5 @@ void UE::AIAssistant::SlateQuerier::STestWidget::Construct(const FArguments& InA
 			 ]
 		 ];
 }
+
+#endif  // WITH_DEV_AUTOMATION_TESTS
