[AIAssistant] Add the ability to register OnConversationUpdate callbacks

diff --git a/Engine/Plugins/Experimental/AIAssistant/Source/AIAssistant/Private/AIAssistantWebApi.cpp b/Engine/Plugins/Experimental/AIAssistant/Source/AIAssistant/Private/AIAssistantWebApi.cpp
--- a/Engine/Plugins/Experimental/AIAssistant/Source/AIAssistant/Private/AIAssistantWebApi.cpp
+++ b/Engine/Plugins/Experimental/AIAssistant/Source/AIAssistant/Private/AIAssistantWebApi.cpp
@@ -20,17 +20,12 @@ namespace UE::AIAssistant
 	const FString FWebApi::WebApiObjectName =
 		UE_AIASSISTANT_WEB_API_PARENT_OBJECT_NAME TEXT(".")
 		UE_AIASSISTANT_WEB_API_CHILD_OBJECT_NAME;
-
-	const FString FWebApi::WebApiAvailableFunction =
-		TEXT("(() => { ")
-			TEXT("return {")
-				TEXT("value: Object.hasOwn(")
-					UE_AIASSISTANT_WEB_API_PARENT_OBJECT_NAME
-					TEXT(", ")
-					TEXT("'") UE_AIASSISTANT_WEB_API_CHILD_OBJECT_NAME TEXT("'")
-				TEXT(")")
-			TEXT("}; ")
-		TEXT("})");
+	const FString FWebApi::WebApiAvailableFunction = FString::Format(
+		TEXT("(() => { return {value: Object.hasOwn({0}, '{1}')}; })"),
+		{
+			UE_AIASSISTANT_WEB_API_PARENT_OBJECT_NAME,
+			UE_AIASSISTANT_WEB_API_CHILD_OBJECT_NAME
+		});
 
 	const FString FWebApi::FunctionCallFormatTemplate = TEXT(R"js(
 try {
@@ -44,6 +39,12 @@ try {
 } catch (error) {
   {NotifyHandlerOfError}
 }
+)js");
+
+	const FString FWebApi::FunctionCallbackFormatTemplate = TEXT(R"js(
+(result) => {
+  {NotifyHandlerOfResult}
+}
 )js");
 
 	FWebApi::FWebApi(
@@ -98,13 +99,42 @@ try {
 	void FWebApi::UpdateGlobalLocale(const FString& LocaleString)
 	{
 		(void)ExecuteFunction(
-			*WebApiObjectName, TEXT("updateGlobalLocale"), *FString::Printf(TEXT("\"%s\""),
-			*LocaleString));
+			*WebApiObjectName, TEXT("updateGlobalLocale"),
+			*FString::Printf(TEXT("\"%s\""), *LocaleString));
+	}
+
+	FString FWebApi::RegisterOnConversationUpdate(
+		TFunction<bool(FConversationId)> OnConversationUpdate)
+	{
+		using FResult = UAIAssistantWebJavaScriptResultDelegate::FResult;
+		FString HandlerId = RegisterCallback(
+			[OnConversationUpdate = MoveTemp(OnConversationUpdate)](FResult Result) {
+			FConversationId Id;
+			if (Result.bJsonIsError || !Id.FromJson(Result.Json))
+			{
+				UE_LOG(LogAIAssistant, Error, TEXT("Javascript Error running callback: %s"), *Result.Json);
+				return true;
+			}
+			return OnConversationUpdate(Id);
+		});
+
+		(void)ExecuteFunction(
+			*WebApiObjectName, TEXT("registerOnConversationUpdate"),
+			*FormatFunctionCallback(HandlerId));
+		return HandlerId;
+	}
+
+	void FWebApi::UnregisterOnConversationUpdate(const FString& HandlerId)
+	{
+		(void)ExecuteFunction(
+			*WebApiObjectName, TEXT("unregisterOnConversationUpdate"),
+			*FString::Printf(TEXT("\"%s\""), *HandlerId));
+		WebJavaScriptResultDelegate->UnregisterResultHandler(HandlerId);
 	}
 
 	FString FWebApi::FormatFunctionCall(
 		const TCHAR* InstanceName, const TCHAR* FunctionName, const TCHAR* Arguments,
-		const FString& HandlerId)
+		const FString& HandlerId) const
 	{
 		check(FunctionName);
 		check(Arguments);
@@ -116,14 +146,24 @@ try {
 		return FString::Format(
 			*FunctionCallFormatTemplate,
 			{
-				{ TEXT("Function"), FString(Function) },
+				{ TEXT("Function"), Function },
 				{ TEXT("Arguments"), Arguments },
 				{ TEXT("NotifyHandlerOfResult"), Handlers.Key },
 				{ TEXT("NotifyHandlerOfError"), Handlers.Value },
 			});
 	}
 
-	TPair<FString, FString> FWebApi::FormatResultAndErrorHandlers(const FString& HandlerId)
+	FString FWebApi::FormatFunctionCallback(const FString& HandlerId) const
+	{
+		auto Handlers = FormatResultAndErrorHandlers(HandlerId);
+		return FString::Format(
+			*FunctionCallbackFormatTemplate,
+			FStringFormatNamedArguments {
+				{ TEXT("NotifyHandlerOfResult"), Handlers.Key }
+			});
+	}
+
+	TPair<FString, FString> FWebApi::FormatResultAndErrorHandlers(const FString& HandlerId) const
 	{
 		if (!HandlerId.IsEmpty())
 		{
@@ -144,11 +184,17 @@ try {
 		return MoveTemp(HandlerIdAndFuture.Value);
 	}
 
+	FString FWebApi::RegisterCallback(
+		TFunction<bool(UAIAssistantWebJavaScriptResultDelegate::FResult)> Callback)
+	{
+		return WebJavaScriptResultDelegate->RegisterResultHandlerForCallback(Callback);
+	}
+
 	void FWebApi::ExecuteAsyncFunction(
 		const TCHAR* InstanceName, const TCHAR* FunctionName, const TCHAR* Arguments,
 		const TCHAR* HandlerId)
 	{
- 		WebJavaScriptExecutor.ExecuteJavaScript(
+		WebJavaScriptExecutor.ExecuteJavaScript(
 			FormatFunctionCall(InstanceName, FunctionName, Arguments, HandlerId));
 	}
 
diff --git a/Engine/Plugins/Experimental/AIAssistant/Source/AIAssistant/Private/AIAssistantWebApi.h b/Engine/Plugins/Experimental/AIAssistant/Source/AIAssistant/Private/AIAssistantWebApi.h
--- a/Engine/Plugins/Experimental/AIAssistant/Source/AIAssistant/Private/AIAssistantWebApi.h
+++ b/Engine/Plugins/Experimental/AIAssistant/Source/AIAssistant/Private/AIAssistantWebApi.h
@@ -243,22 +243,39 @@ namespace UE::AIAssistant
 		// Set agent environment for the conversational UI.
 		void SetAgentEnvironment(const FAgentEnvironmentId& AgentEnvironmentId);
 
+		// Update the global locale state.
 		void UpdateGlobalLocale(const FString& LocaleString);
 
+		// Register for conversation updates. Return true to unregister, or explicitly call
+		// UnregisterOnConversationUpdate.
+		FString RegisterOnConversationUpdate(
+			TFunction<bool(FConversationId)> OnConversationUpdate);
+
+		// Unregister conversation update callback using the HandlerId returned from
+		// RegisterOnConversationUpdate.
+		void UnregisterOnConversationUpdate(const FString& HandlerId);
+
 	protected:
 		// Format a function call of a member with result handling.
 		FString FormatFunctionCall(
 			const TCHAR* InstanceName, const TCHAR* FunctionName,
-			const TCHAR* Arguments = TEXT(""),
-			const FString& HandlerId = FString());
+			const TCHAR* Arguments = TEXT(""), const FString& HandlerId = FString()) const;
+
+		// Format a function callback which calls a given handler ID with a result.
+		FString FormatFunctionCallback(const FString& HandlerId) const;
 
 		// Format JavaScript for result and error handler function calls.
-		TPair<FString, FString> FormatResultAndErrorHandlers(const FString& HandlerId);
+		TPair<FString, FString> FormatResultAndErrorHandlers(const FString& HandlerId) const;
 
 		// Execute a javascript function getting the result as a JSON encoded string.
 		TFuture<UAIAssistantWebJavaScriptResultDelegate::FResult> ExecuteFunction(
 			const TCHAR* InstanceName, const TCHAR* FunctionName,
 			const TCHAR* Arguments = TEXT(""));
+		
+		// Register a callback with Javascript, allowing the given TFunction to be called multiple
+		// times from Javascript. To unregister, 
+		FString RegisterCallback(
+			TFunction<bool(UAIAssistantWebJavaScriptResultDelegate::FResult)> Callback);
 
 		// Execute a javascript function handling the result with the specified handler.
 		// NOTE: This is virtual so tests can capture the function, arguments and handler ID before
@@ -379,5 +396,6 @@ namespace UE::AIAssistant
 		// * NotifyHandlerOfError: JavaScript snippet that is called with "error" to handle
 		//   a function error.
 		static const FString FunctionCallFormatTemplate;
+		static const FString FunctionCallbackFormatTemplate;
 	};
 }
diff --git a/Engine/Plugins/Experimental/AIAssistant/Source/AIAssistant/Private/AIAssistantWebJavaScriptResultDelegate.cpp b/Engine/Plugins/Experimental/AIAssistant/Source/AIAssistant/Private/AIAssistantWebJavaScriptResultDelegate.cpp
--- a/Engine/Plugins/Experimental/AIAssistant/Source/AIAssistant/Private/AIAssistantWebJavaScriptResultDelegate.cpp
+++ b/Engine/Plugins/Experimental/AIAssistant/Source/AIAssistant/Private/AIAssistantWebJavaScriptResultDelegate.cpp
@@ -38,6 +38,12 @@ void UAIAssistantWebJavaScriptResultDelegate::RegisterResultHandlerWithId(
 	ResultHandlersById.Emplace(HandlerId, MakeShared<FResultHandler>(MoveTemp(Handler)));
 }
 
+void UAIAssistantWebJavaScriptResultDelegate::UnregisterResultHandler(const FString& HandlerId)
+{
+	UE::TUniqueLock Lock(ResultHandlersByIdLock);
+	ResultHandlersById.Remove(HandlerId);
+}
+
 FString UAIAssistantWebJavaScriptResultDelegate::RegisterResultHandler(FResultHandler&& Handler)
 {
 	FString HandlerId = CreateHandlerId();
@@ -60,6 +66,19 @@ UAIAssistantWebJavaScriptResultDelegate::RegisterResultHandlerForFuture()
 	return TPair<FString, TFuture<FResult>>(MoveTemp(HandlerId), MoveTemp(Future));
 }
 
+FString UAIAssistantWebJavaScriptResultDelegate::RegisterResultHandlerForCallback(
+	TFunction<bool(FResult)> Callback)
+{
+	FString HandlerId = CreateHandlerId();
+	RegisterResultHandlerWithId(
+		[Callback=MoveTemp(Callback)](FResultHandlerContext ResultHandlerContext) mutable -> bool
+		{
+			return Callback(MoveTemp(ResultHandlerContext));
+		},
+		HandlerId);
+	return HandlerId;
+}
+
 void UAIAssistantWebJavaScriptResultDelegate::HandleResult(
 	const FString& HandlerId, const FString& ResultJson, bool bResultJsonIsError)
 {
@@ -74,8 +93,7 @@ void UAIAssistantWebJavaScriptResultDelegate::HandleResult(
 		(*Handler)(FResultHandlerContext{ { ResultJson, bResultJsonIsError }, HandlerId });
 	if (RemoveHandler)
 	{
-		UE::TUniqueLock Lock(ResultHandlersByIdLock);
-		ResultHandlersById.Remove(HandlerId);
+		UnregisterResultHandler(HandlerId);
 	}
 }
 
diff --git a/Engine/Plugins/Experimental/AIAssistant/Source/AIAssistant/Private/AIAssistantWebJavaScriptResultDelegate.h b/Engine/Plugins/Experimental/AIAssistant/Source/AIAssistant/Private/AIAssistantWebJavaScriptResultDelegate.h
--- a/Engine/Plugins/Experimental/AIAssistant/Source/AIAssistant/Private/AIAssistantWebJavaScriptResultDelegate.h
+++ b/Engine/Plugins/Experimental/AIAssistant/Source/AIAssistant/Private/AIAssistantWebJavaScriptResultDelegate.h
@@ -74,6 +74,13 @@ public:
 	// This handler will always only be executed once.
 	TPair<FString, TFuture<FResult>> RegisterResultHandlerForFuture();
 
+	// Register a callback that may be called multiple times. Return true from the callback to
+	// unregister the handler, or explicitly call UnregisterResultHandler.
+	FString RegisterResultHandlerForCallback(TFunction<bool(FResult)> Callback);
+
+	// Unregister a previously registered handler with the specified ID.
+	void UnregisterResultHandler(const FString& HandlerId);
+
 	// Generate a FString::Format() format string to call the JavaScript function associated with
 	// the handler using the specified ID. The returned format string has two arguments "{0}"
 	// that expects to accept an object that can be serialized as JSON to pass to the
diff --git a/Engine/Plugins/Experimental/AIAssistant/Source/AIAssistant/Private/Tests/AIAssistantFakeWebApi.cpp b/Engine/Plugins/Experimental/AIAssistant/Source/AIAssistant/Private/Tests/AIAssistantFakeWebApi.cpp
--- a/Engine/Plugins/Experimental/AIAssistant/Source/AIAssistant/Private/Tests/AIAssistantFakeWebApi.cpp
+++ b/Engine/Plugins/Experimental/AIAssistant/Source/AIAssistant/Private/Tests/AIAssistantFakeWebApi.cpp
@@ -80,6 +80,13 @@ namespace UE::AIAssistant
 			ResultJson, bResultJsonIsError);
 	}
 
+	void FFakeWebApi::InvokeCallback(
+		const FString& HandlerId, const FString& ResultJson, bool bResultJsonIsError)
+	{
+		FWebJavaScriptResultDelegateAccessor::CallHandleResult(
+			*WebJavaScriptResultDelegate, HandlerId, ResultJson, bResultJsonIsError);
+	}
+
 	void FFakeWebApi::ExecuteAsyncFunction(
 		const TCHAR* InstanceName, const TCHAR* FunctionName, const TCHAR* Arguments,
 		const TCHAR* HandlerId)
diff --git a/Engine/Plugins/Experimental/AIAssistant/Source/AIAssistant/Private/Tests/AIAssistantFakeWebApi.h b/Engine/Plugins/Experimental/AIAssistant/Source/AIAssistant/Private/Tests/AIAssistantFakeWebApi.h
--- a/Engine/Plugins/Experimental/AIAssistant/Source/AIAssistant/Private/Tests/AIAssistantFakeWebApi.h
+++ b/Engine/Plugins/Experimental/AIAssistant/Source/AIAssistant/Private/Tests/AIAssistantFakeWebApi.h
@@ -108,6 +108,10 @@ namespace UE::AIAssistant
 			const FExecutedAsyncFunction& ExecutedAsyncFunction,
 			const FString& ResultJson, bool bResultJsonIsError);
 
+		// Invoke a registered callback.
+		void InvokeCallback(
+			const FString& HandlerId, const FString& ResultJson, bool bResultJsonIsError);
+
 	protected:
 		void ExecuteAsyncFunction(
 			const TCHAR* InstanceName, const TCHAR* FunctionName, const TCHAR* Arguments,
diff --git a/Engine/Plugins/Experimental/AIAssistant/Source/AIAssistant/Private/Tests/AIAssistantWebApiTest.cpp b/Engine/Plugins/Experimental/AIAssistant/Source/AIAssistant/Private/Tests/AIAssistantWebApiTest.cpp
--- a/Engine/Plugins/Experimental/AIAssistant/Source/AIAssistant/Private/Tests/AIAssistantWebApiTest.cpp
+++ b/Engine/Plugins/Experimental/AIAssistant/Source/AIAssistant/Private/Tests/AIAssistantWebApiTest.cpp
@@ -42,6 +42,11 @@ namespace UE::AIAssistant
 			return WebApi.FormatFunctionCall(InstanceName, FunctionName, Arguments, HandlerId);
 		}
 
+		static FString FormatFunctionCallback(
+			FWebApi& WebApi, const FString& HandlerId) {
+			return WebApi.FormatFunctionCallback(HandlerId);
+		}
+
 		static TPair<FString, FString> FormatResultAndErrorHandlers(
 			FWebApi& WebApi, const FString& HandlerId)
 		{
@@ -197,6 +202,51 @@ try {
 			}));
 }
 
+IMPLEMENT_SIMPLE_AUTOMATION_TEST(
+	FAIAssistantWebApiTestFormatFunctionCallback,
+	"AI.Assistant.WebApi.FormatFunctionCallback",
+	AIAssistantTest::Flags);
+
+bool FAIAssistantWebApiTestFormatFunctionCallback::RunTest(const FString& UnusedParameters)
+{
+	FFakeWebApiContainer WebApi;
+	return TestEqual(
+		TEXT("FormatFunctionCallback"),
+		FWebApiAccessor::FormatFunctionCallback(*WebApi, TEXT("")),
+		TEXT(R"js(
+(result) => {
+  
+}
+)js"));
+}
+
+
+IMPLEMENT_SIMPLE_AUTOMATION_TEST(
+	FAIAssistantWebApiTestFormatFunctionCallbackWithResultHandler,
+	"AI.Assistant.WebApi.FormatFunctionCallbackWithResultHandler",
+	AIAssistantTest::Flags);
+
+bool FAIAssistantWebApiTestFormatFunctionCallbackWithResultHandler::RunTest(
+	const FString& UnusedParameters)
+{
+	FFakeWebApiContainer WebApi;
+	const FString FakeHandlerId = TEXT("foobar");
+	const auto Handlers = FWebApiAccessor::FormatResultAndErrorHandlers(*WebApi, FakeHandlerId);
+	return TestEqual(
+		TEXT("FormatFunctionCallWithResultHandler"),
+		FWebApiAccessor::FormatFunctionCallback(*WebApi, FakeHandlerId),
+		FString::Format(
+			TEXT(R"js(
+(result) => {
+  {HandleResult}
+}
+)js"),
+			{
+				{ TEXT("HandleResult"), *Handlers.Key },
+				{ TEXT("HandleError"), *Handlers.Value },
+			}));
+}
+
 IMPLEMENT_SIMPLE_AUTOMATION_TEST(
 	FAIAssistantWebApiTestIsAvailable,
 	"AI.Assistant.WebApi.IsAvailable",
@@ -315,4 +365,174 @@ bool FAIAssistantWebApiTestUpdateGlobalLocale::RunTest(const FString& UnusedPara
 		*LocaleFromSetting)).IsValid();
 }
 
+IMPLEMENT_SIMPLE_AUTOMATION_TEST(
+	FAIAssistantWebApiTestRegisterOnConversationUpdate,
+	"AI.Assistant.WebApi.RegisterOnConversationUpdate",
+	AIAssistantTest::Flags);
+
+bool FAIAssistantWebApiTestRegisterOnConversationUpdate::RunTest(const FString& UnusedParameters)
+{
+	FFakeWebApiContainer WebApi;
+	int CallCount = 0;
+	FString HandlerId = WebApi->RegisterOnConversationUpdate([&CallCount](FConversationId ConversationId)
+		{
+			CallCount += 1;
+			return false;
+		});
+	
+	return
+		WebApi->TestExpectAsyncFunctionCall(
+			*this,
+			nullptr,
+			TEXT("RegisterOnConversationUpdate"),
+			*FWebApiAccessor::FormatFunctionCallback(*WebApi, *HandlerId)).IsValid()
+		&& TestEqual(TEXT("RegisterOnConversationUpdate_NoCalls"), CallCount, 0);
+}
+
+IMPLEMENT_SIMPLE_AUTOMATION_TEST(
+	FAIAssistantWebApiTestRegisterOnConversationUpdateAndInvoke,
+	"AI.Assistant.WebApi.RegisterOnConversationUpdateAndInvoke",
+	AIAssistantTest::Flags);
+
+bool FAIAssistantWebApiTestRegisterOnConversationUpdateAndInvoke::RunTest(
+	const FString& UnusedParameters)
+{
+	FFakeWebApiContainer WebApi;
+	FString ActualConversationId;
+	int CallCount = 0;
+	FString CallbackHandle = WebApi->RegisterOnConversationUpdate(
+		[&ActualConversationId, &CallCount](FConversationId ConversationId)
+		{
+			ActualConversationId = ConversationId.Id;
+			CallCount += 1;
+			return false;
+		});
+	if (!WebApi->TestExpectAsyncFunctionCallAndComplete(
+		*this,
+		nullptr,
+		TEXT("registerOnConversationUpdate"),
+		*FWebApiAccessor::FormatFunctionCallback(*WebApi, *CallbackHandle),
+		TEXT(""),
+		false))
+	{
+		return false;
+	}
+	FString ExpectedConversationId = TEXT("foobar");
+	WebApi->InvokeCallback(
+		CallbackHandle,
+		*FString::Printf(TEXT(R"js({ "id": "%s" })js"), *ExpectedConversationId),
+		false);
+	return
+		TestEqual(
+			TEXT("RegisterOnConversationUpdate_ConversationId"),
+			ActualConversationId,
+			ExpectedConversationId)
+		&& TestEqual(TEXT("RegisterOnConversationUpdate_CallCount"), CallCount, 1);
+}
+
+IMPLEMENT_SIMPLE_AUTOMATION_TEST(
+	FAIAssistantWebApiTestRegisterOnConversationUpdateAndInvokeTwice,
+	"AI.Assistant.WebApi.RegisterOnConversationUpdateAndInvokeTwice",
+	AIAssistantTest::Flags);
+
+bool FAIAssistantWebApiTestRegisterOnConversationUpdateAndInvokeTwice::RunTest(
+	const FString& UnusedParameters)
+{
+	FFakeWebApiContainer WebApi;
+	TArray<FString> ActualConversationIds;
+	int CallCount = 0;
+	FString CallbackHandle = WebApi->RegisterOnConversationUpdate(
+		[&ActualConversationIds, &CallCount](FConversationId ConversationId)
+		{
+			ActualConversationIds.Push(ConversationId.Id);
+			CallCount += 1;
+			return false;
+		});
+
+	if (!WebApi->TestExpectAsyncFunctionCallAndComplete(
+		*this,
+		nullptr,
+		TEXT("RegisterOnConversationUpdate"),
+		*FWebApiAccessor::FormatFunctionCallback(*WebApi, *CallbackHandle),
+		TEXT(""),
+		false))
+	{
+		return false;
+	}
+	TArray<FString> ExpectedConversationIds = { TEXT("foo"), TEXT("bar") };
+	for (const auto& ConversationId : ExpectedConversationIds)
+	{
+		WebApi->InvokeCallback(
+			CallbackHandle,
+			*FString::Printf(TEXT(R"js({ "id": "%s" })js"), *ConversationId),
+			false);
+	}
+	return
+		TestEqual(
+			TEXT("RegisterOnConversationUpdateTwice_ConversationIds"),
+			ActualConversationIds,
+			ExpectedConversationIds)
+		&& TestEqual(TEXT("RegisterOnConversationUpdateTwice_CallCount"), CallCount, 2);
+}
+
+IMPLEMENT_SIMPLE_AUTOMATION_TEST(
+	FAIAssistantWebApiTestUnegisterOnConversationUpdateViaReturnValue,
+	"AI.Assistant.WebApi.UnregisterOnConversationUpdateViaReturnValue",
+	AIAssistantTest::Flags);
+
+
+bool FAIAssistantWebApiTestUnegisterOnConversationUpdateViaReturnValue::RunTest(
+	const FString& UnusedParameters)
+{
+	FFakeWebApiContainer WebApi;
+	TArray<FString> ActualConversationIds;
+	int CallCount = 0;
+	FString CallbackHandle = WebApi->RegisterOnConversationUpdate(
+		[&ActualConversationIds, &CallCount](FConversationId ConversationId)
+		{
+			ActualConversationIds.Push(ConversationId.Id);
+			CallCount += 1;
+			return true;
+		});
+
+	if (!WebApi->TestExpectAsyncFunctionCallAndComplete(
+		*this,
+		nullptr,
+		TEXT("RegisterOnConversationUpdate"),
+		*FWebApiAccessor::FormatFunctionCallback(*WebApi, *CallbackHandle),
+		TEXT(""),
+		false))
+	{
+		return false;
+	}
+	TArray<FString> ExpectedConversationIds = { TEXT("foo") };
+	for (const auto& ConversationId : ExpectedConversationIds)
+	{
+		WebApi->InvokeCallback(
+			CallbackHandle,
+			*FString::Printf(TEXT(R"js({ "id": "%s" })js"), *ConversationId),
+			false);
+	}
+	return
+		TestEqual(
+			TEXT("RegisterOnConversationUpdateTwice_ConversationIds"),
+			ActualConversationIds,
+			ExpectedConversationIds)
+		&& TestEqual(TEXT("RegisterOnConversationUpdateTwice_CallCount"), CallCount, 1);
+}
+
+IMPLEMENT_SIMPLE_AUTOMATION_TEST(
+	FAIAssistantWebApiTestUnregisterOnConversationUpdate,
+	"AI.Assistant.WebApi.UnregisterOnConversationUpdate",
+	AIAssistantTest::Flags);
+bool FAIAssistantWebApiTestUnregisterOnConversationUpdate::RunTest(const FString& UnusedParameters)
+{
+	FFakeWebApiContainer WebApi;
+	FString FakeHandlerId = TEXT("foobar");
+	WebApi->UnregisterOnConversationUpdate(FakeHandlerId);
+	return WebApi->TestExpectAsyncFunctionCall(
+		*this, nullptr, TEXT("unregisterOnConversationUpdate"),
+		*FString::Printf(TEXT("\"%s\""), *FakeHandlerId)).IsValid();
+}
+
 #endif  // WITH_DEV_AUTOMATION_TESTS
\ No newline at end of file
diff --git a/Engine/Plugins/Experimental/AIAssistant/Source/AIAssistant/Private/Tests/AIAssistantWebJavaScriptResultDelegateTest.cpp b/Engine/Plugins/Experimental/AIAssistant/Source/AIAssistant/Private/Tests/AIAssistantWebJavaScriptResultDelegateTest.cpp
--- a/Engine/Plugins/Experimental/AIAssistant/Source/AIAssistant/Private/Tests/AIAssistantWebJavaScriptResultDelegateTest.cpp
+++ b/Engine/Plugins/Experimental/AIAssistant/Source/AIAssistant/Private/Tests/AIAssistantWebJavaScriptResultDelegateTest.cpp
@@ -237,6 +237,62 @@ bool FAIAssistantWebJavaScriptResultDelegateTestRegisterResultHandlerForFuture::
 	return true;
 }
 
+IMPLEMENT_SIMPLE_AUTOMATION_TEST(
+	FAIAssistantWebJavaScriptResultDelegateTestRegisterResultHandlerForCallback,
+	"AI.Assistant.WebJavaScriptResultDelegate.RegisterResultHandlerForCallback",
+	AIAssistantTest::Flags);
+
+bool FAIAssistantWebJavaScriptResultDelegateTestRegisterResultHandlerForCallback::RunTest(
+	const FString& UnusedParameters)
+{
+	TStrongObjectPtr<UAIAssistantWebJavaScriptResultDelegate> ResultDelegate(
+		NewObject<UAIAssistantWebJavaScriptResultDelegate>());
+	int CallCount = 0;
+	FString HandlerId = ResultDelegate->RegisterResultHandlerForCallback(
+		[&CallCount](UAIAssistantWebJavaScriptResultDelegate::FResult) mutable -> bool
+		{
+			CallCount++;
+			return false;
+		});
+	const int ExpectedCalls = 3;
+	for (int i = 0; i < ExpectedCalls; ++i)
+	{
+		FWebJavaScriptResultDelegateAccessor::CallHandleResult(
+			*ResultDelegate, HandlerId, TEXT(""), false);
+		(void)TestEqual(
+			TEXT("ExpectedNumberOfCalls"), i + 1, CallCount);
+	}
+	return true;
+}
+
+IMPLEMENT_SIMPLE_AUTOMATION_TEST(
+	FAIAssistantWebJavaScriptResultDelegateTestUnregisterResultHandler,
+	"AI.Assistant.WebJavaScriptResultDelegate.UnregisterResultHandler",
+	AIAssistantTest::Flags);
+
+bool FAIAssistantWebJavaScriptResultDelegateTestUnregisterResultHandler::RunTest(
+	const FString& UnusedParameters)
+{
+	TStrongObjectPtr<UAIAssistantWebJavaScriptResultDelegate> ResultDelegate(
+		NewObject<UAIAssistantWebJavaScriptResultDelegate>());
+	bool bWasCalled = false;
+	FString HandlerId = ResultDelegate->RegisterResultHandler(
+		[&bWasCalled](UAIAssistantWebJavaScriptResultDelegate::FResultHandlerContext&&) -> bool
+		{
+			bWasCalled = true;
+			return true;
+		});
+	ResultDelegate->UnregisterResultHandler(HandlerId);
+	(void)TestEqual(
+		TEXT("NumberOfHandlers"), 0,
+		FWebJavaScriptResultDelegateAccessor::GetResultHandlersById(
+			*ResultDelegate).Num());
+	FWebJavaScriptResultDelegateAccessor::CallHandleResult(
+		*ResultDelegate, HandlerId, TEXT(""), false);
+	(void)TestFalse(TEXT("HandlerNotCalled"), bWasCalled);
+	return true;
+}
+
 IMPLEMENT_SIMPLE_AUTOMATION_TEST(
 	FAIAssistantWebJavaScriptResultDelegateTestCompleteAllPendingPromises,
 	"AI.Assistant.WebJavaScriptResultDelegate.CompleteAllPendingPromises",
