[AI Assistant] Prefer passing around FString rather than FJsonObject in IToolsetHandler

Reason: avoid redundant serialization and validation passes on the JSON schema.
Also simplifies implementing UFUNCTION wrapper methods.

#rnx
#Synced-CL 48281750

diff --git a/Engine/Plugins/Experimental/AIAssistant/Source/AIAssistant/Private/AIAssistantToolsetHandler.h b/Engine/Plugins/Experimental/AIAssistant/Source/AIAssistant/Private/AIAssistantToolsetHandler.h
--- a/Engine/Plugins/Experimental/AIAssistant/Source/AIAssistant/Private/AIAssistantToolsetHandler.h
+++ b/Engine/Plugins/Experimental/AIAssistant/Source/AIAssistant/Private/AIAssistantToolsetHandler.h
@@ -4,7 +4,6 @@
 
 #include "Async/Future.h"
 #include "Containers/UnrealString.h"
-#include "Dom/JsonObject.h"
 #include "Templates/ValueOrError.h"
 
 namespace UE::AIAssistant
@@ -17,11 +16,11 @@ namespace UE::AIAssistant
 
 		// Executes a tool with the given name and JSON input, returning a future with the JSON
 		// output.
-		virtual TFuture<TValueOrError<FJsonObject, FString>> ExecuteTool(
-			const FString& ToolName, const FJsonObject& JsonInput) = 0;
+		virtual TFuture<TValueOrError<FString, FString>> ExecuteTool(
+			const FString& ToolName, const FString& JsonInput) = 0;
 
 		// Returns the JSON schema for the toolset
-		virtual FJsonObject GetToolsetJsonSchema() const = 0;
+		virtual FString GetToolsetJsonSchema() const = 0;
 
 		// Returns the name of the toolset
 		virtual FString GetToolsetName() const = 0;
diff --git a/Engine/Plugins/Experimental/AIAssistant/Source/AIAssistant/Private/AIAssistantToolsetRegistry.cpp b/Engine/Plugins/Experimental/AIAssistant/Source/AIAssistant/Private/AIAssistantToolsetRegistry.cpp
--- a/Engine/Plugins/Experimental/AIAssistant/Source/AIAssistant/Private/AIAssistantToolsetRegistry.cpp
+++ b/Engine/Plugins/Experimental/AIAssistant/Source/AIAssistant/Private/AIAssistantToolsetRegistry.cpp
@@ -96,15 +96,15 @@ namespace UE::AIAssistant
 		return ToolsetHandler;
 	}
 
-	TFuture<TValueOrError<FJsonObject, FString>> FToolsetRegistry::ExecuteTool(
-		const FString& ToolsetAndToolName, const FJsonObject& JsonInput)
+	TFuture<TValueOrError<FString, FString>> FToolsetRegistry::ExecuteTool(
+		const FString& ToolsetAndToolName, const FString& JsonInput)
 	{
 		// Split the tool name to get the toolset name and the actual tool name.
 		TValueOrError<FToolDescriptor, FString> ToolDescriptor =
 			FToolDescriptor::FromString(ToolsetAndToolName, true);
 		if (ToolDescriptor.HasError())
 		{
-			return MakeFulfilledPromise<TValueOrError<FJsonObject, FString>>(
+			return MakeFulfilledPromise<TValueOrError<FString, FString>>(
 				MakeError(ToolDescriptor.StealError())).GetFuture();
 		}
 
@@ -112,15 +112,15 @@ namespace UE::AIAssistant
 		return ExecuteTool(ToolDescriptor.StealValue(), JsonInput);
 	}
 
-	TFuture<TValueOrError<FJsonObject, FString>> FToolsetRegistry::ExecuteTool(
-		const FToolDescriptor& ToolDescriptor, const FJsonObject& JsonInput)
+	TFuture<TValueOrError<FString, FString>> FToolsetRegistry::ExecuteTool(
+		const FToolDescriptor& ToolDescriptor, const FString& JsonInput)
 	{
 		FString ErrorMessage;
 		TSharedPtr<IToolsetHandler>* ToolsetHandler =
 			Find(ToolDescriptor.ToolsetName, true, &ErrorMessage);
 		if (!ToolsetHandler)
 		{
-			return MakeFulfilledPromise<TValueOrError<FJsonObject, FString>>(
+			return MakeFulfilledPromise<TValueOrError<FString, FString>>(
 				MakeError(ErrorMessage)).GetFuture();
 		}
 
diff --git a/Engine/Plugins/Experimental/AIAssistant/Source/AIAssistant/Private/AIAssistantToolsetRegistry.h b/Engine/Plugins/Experimental/AIAssistant/Source/AIAssistant/Private/AIAssistantToolsetRegistry.h
--- a/Engine/Plugins/Experimental/AIAssistant/Source/AIAssistant/Private/AIAssistantToolsetRegistry.h
+++ b/Engine/Plugins/Experimental/AIAssistant/Source/AIAssistant/Private/AIAssistantToolsetRegistry.h
@@ -4,7 +4,6 @@
 
 #include "Containers/Map.h"
 #include "Containers/UnrealString.h"
-#include "Dom/JsonObject.h"
 #include "Misc/Variant.h"
 #include "Templates/SharedPointer.h"
 
@@ -41,15 +40,15 @@ namespace UE::AIAssistant
 		// 
 		// If there is an error in attempting to execute a tool, such as a tool not existing,
 		// the future will contain an error string.
-		TFuture<TValueOrError<FJsonObject, FString>> ExecuteTool(
-			const FString& ToolName, const FJsonObject& JsonInput);
+		TFuture<TValueOrError<FString, FString>> ExecuteTool(
+			const FString& ToolName, const FString& JsonInput);
 
 		// Execute a tool from a registered toolset handler.
 		// 
 		// If there is an error in attempting to execute a tool, such as a tool not existing,
 		// the future will contain an error string.
-		TFuture<TValueOrError<FJsonObject, FString>> ExecuteTool(
-			const FToolDescriptor& ToolDescriptor, const FJsonObject& JsonInput);
+		TFuture<TValueOrError<FString, FString>> ExecuteTool(
+			const FToolDescriptor& ToolDescriptor, const FString& JsonInput);
 
 	private:
 		// The set of registered toolset handlers.
diff --git a/Engine/Plugins/Experimental/AIAssistant/Source/AIAssistant/Private/Tests/AIAssistantToolsetRegistryTest.cpp b/Engine/Plugins/Experimental/AIAssistant/Source/AIAssistant/Private/Tests/AIAssistantToolsetRegistryTest.cpp
--- a/Engine/Plugins/Experimental/AIAssistant/Source/AIAssistant/Private/Tests/AIAssistantToolsetRegistryTest.cpp
+++ b/Engine/Plugins/Experimental/AIAssistant/Source/AIAssistant/Private/Tests/AIAssistantToolsetRegistryTest.cpp
@@ -2,7 +2,6 @@
 
 #include "AIAssistantToolsetRegistry.h"
 
-#include "Dom/JsonObject.h"
 #include "Misc/AutomationTest.h"
 #include "UObject/StrongObjectPtr.h"
 
@@ -114,24 +113,25 @@ bool FAIAssistantToolsetDescriptorFromStringMultipleDots::RunTest(const FString&
 class FakeToolset : public IToolsetHandler
 {
 public:
-	TFuture<TValueOrError<FJsonObject, FString>> ExecuteTool(
-		const FString& ToolName, const FJsonObject& JsonInput) override
+	static const inline FString SUCCESSFUL_RESULT = FString(TEXT(R"({"result": "success"})"));
+
+	TFuture<TValueOrError<FString, FString>> ExecuteTool(
+		const FString& ToolName, const FString& JsonInput) override
 	{
 		if (ToolName != TEXT("SomeTool"))
 		{
-			return MakeFulfilledPromise<TValueOrError<FJsonObject, FString>>(
+			return MakeFulfilledPromise<TValueOrError<FString, FString>>(
 				MakeError(FString(TEXT("Tool not found")))).GetFuture();
 		}
 
-		FJsonObject Result;
-		Result.SetStringField("result", "success");
-		return MakeFulfilledPromise<TValueOrError<FJsonObject, FString>>(
+		FString Result(SUCCESSFUL_RESULT);
+		return MakeFulfilledPromise<TValueOrError<FString, FString>>(
 			MakeValue(Result)).GetFuture();
 	}
 
-	FJsonObject GetToolsetJsonSchema() const override
+	FString GetToolsetJsonSchema() const override
 	{
-		return FJsonObject();
+		return FString();
 	}
 
 	FString GetToolsetName() const override
@@ -355,18 +355,17 @@ const FString& UnusedParameters)
 		TEXT("RegisterToolsetHandler_Valid"),
 		Registry->RegisterToolset(ValidToolsetHandler));
 
-	FJsonObject InputArguments;
-	InputArguments.SetStringField(TEXT("input"), TEXT("value"));
+	FString InputArguments(TEXT(R"({\"input\": \"value\"})"));
 
 	// Test execution of a tool from the registered toolset handler.
 	Registry->ExecuteTool(TEXT("Fake.SomeTool"), InputArguments)
-		.Next([this](TValueOrError<FJsonObject, FString> Result)
+		.Next([this](TValueOrError<FString, FString> Result)
 		{
 			(void)TestTrue(TEXT("ExecuteTool_SuccessValue"), Result.HasValue());
 			(void)TestEqual(
 				TEXT("ExecuteTool_ResultJson"),
-				*Result.StealValue().GetStringField(TEXT("result")),
-				TEXT("success"));
+				Result.StealValue(),
+				FakeToolset::SUCCESSFUL_RESULT);
 		})
 		.Wait();
 
@@ -390,8 +389,8 @@ bool FAIAssistantToolsetRegistryToolsetHandlerExecuteToolNonexistantTool::RunTes
 		Registry->RegisterToolset(ValidToolsetHandler));
 
 	// Test execution of an invalid tool from the registered toolset handler.
-	Registry->ExecuteTool(TEXT("Fake.SomeNonexistantTool"), FJsonObject())
-		.Next([this](TValueOrError<FJsonObject, FString> Result)
+	Registry->ExecuteTool(TEXT("Fake.SomeNonexistantTool"), FString())
+		.Next([this](TValueOrError<FString, FString> Result)
 		{
 			(void)TestTrue(TEXT("NonexistantToolset_HasError"), Result.HasError());
 			(void)TestEqual(
@@ -422,8 +421,8 @@ bool FAIAssistantToolsetRegistryToolsetHandlerExecuteToolNonexistantToolset::Run
 
 	// Test execution of an invalid toolset.
 	AddExpectedError(TEXT("Toolset 'InvalidTool' not found"));
-	Registry->ExecuteTool(TEXT("InvalidTool.SolveHaltingProblem"), FJsonObject())
-		.Next([this](TValueOrError<FJsonObject, FString> Result)
+	Registry->ExecuteTool(TEXT("InvalidTool.SolveHaltingProblem"), FString())
+		.Next([this](TValueOrError<FString, FString> Result)
 			{
 				(void)TestTrue(TEXT("ExecuteTool_InvalidToolset_HasError"), Result.HasError());
 				(void)TestEqual(
@@ -454,8 +453,8 @@ bool FAIAssistantToolsetRegistryToolsetHandlerExecuteToolBadToolnameFormat::RunT
 
 	// Test execution of an incorrectly formatted tool name.
 	AddExpectedError(TEXT("Invalid tool name format. Expected format: ToolsetName.ToolName"));
-	Registry->ExecuteTool(TEXT("ToolCallThatDoesNotHaveADot"), FJsonObject())
-		.Next([this](TValueOrError<FJsonObject, FString> Result)
+	Registry->ExecuteTool(TEXT("ToolCallThatDoesNotHaveADot"), FString())
+		.Next([this](TValueOrError<FString, FString> Result)
 			{
 				(void)TestTrue(TEXT("ExecuteTool_InvalidToolName_HasError"), Result.HasError());
 				(void)TestEqual(
