[AIAssistant] Map Javascript callback handles to C++ HandlerIds for callback unregistration

#Synced-CL 48261597

diff --git a/Engine/Plugins/Experimental/AIAssistant/Source/AIAssistant/Private/AIAssistantWebApi.cpp b/Engine/Plugins/Experimental/AIAssistant/Source/AIAssistant/Private/AIAssistantWebApi.cpp
--- a/Engine/Plugins/Experimental/AIAssistant/Source/AIAssistant/Private/AIAssistantWebApi.cpp
+++ b/Engine/Plugins/Experimental/AIAssistant/Source/AIAssistant/Private/AIAssistantWebApi.cpp
@@ -103,32 +103,49 @@ try {
 			*FString::Printf(TEXT("\"%s\""), *LocaleString));
 	}
 
-	FString FWebApi::RegisterOnConversationUpdate(
+	TFuture<TValueOrError<FWebApi::FConversationUpdateHandle, FString>>
+		FWebApi::RegisterOnConversationUpdate(
 		TFunction<bool(FConversationId)> OnConversationUpdate)
 	{
 		using FResult = UAIAssistantWebJavaScriptResultDelegate::FResult;
 		FString HandlerId = RegisterCallback(
 			[OnConversationUpdate = MoveTemp(OnConversationUpdate)](FResult Result) {
-			FConversationId Id;
-			if (Result.bJsonIsError || !Id.FromJson(Result.Json))
-			{
-				UE_LOG(LogAIAssistant, Error, TEXT("Javascript Error running callback: %s"), *Result.Json);
-				return true;
-			}
-			return OnConversationUpdate(Id);
-		});
+				FConversationId Id;
+				if (Result.bJsonIsError || !Id.FromJson(Result.Json))
+				{
+					UE_LOG(
+						LogAIAssistant, Error, TEXT("Javascript Error running callback: %s"),
+						*Result.Json);
+					return true;
+				}
+				return OnConversationUpdate(Id);
+			});
 
-		(void)ExecuteFunction(
+		TFuture<FResult> Future = ExecuteFunction(
 			*WebApiObjectName, TEXT("registerOnConversationUpdate"),
 			*FormatFunctionCallback(HandlerId));
-		return HandlerId;
+
+		return Future.Next(
+			[HandlerId](const FResult& Result) -> TValueOrError<FConversationUpdateHandle, FString>
+			{
+				if (Result.bJsonIsError)
+				{
+					return MakeError(Result.Json);
+				}
+				return MakeValue(FConversationUpdateHandle
+					{
+						HandlerId,
+						Result.Json
+					});
+			});
 	}
 
-	void FWebApi::UnregisterOnConversationUpdate(const FString& HandlerId)
+	void FWebApi::UnregisterOnConversationUpdate(const FConversationUpdateHandle& Handle)
 	{
+		auto& [HandlerId, CallbackHandle] = Handle;
 		(void)ExecuteFunction(
 			*WebApiObjectName, TEXT("unregisterOnConversationUpdate"),
-			*FString::Printf(TEXT("\"%s\""), *HandlerId));
+			*FString::Printf(TEXT("%s"), *CallbackHandle));
 		WebJavaScriptResultDelegate->UnregisterResultHandler(HandlerId);
 	}
 
diff --git a/Engine/Plugins/Experimental/AIAssistant/Source/AIAssistant/Private/AIAssistantWebApi.h b/Engine/Plugins/Experimental/AIAssistant/Source/AIAssistant/Private/AIAssistantWebApi.h
--- a/Engine/Plugins/Experimental/AIAssistant/Source/AIAssistant/Private/AIAssistantWebApi.h
+++ b/Engine/Plugins/Experimental/AIAssistant/Source/AIAssistant/Private/AIAssistantWebApi.h
@@ -246,14 +246,17 @@ namespace UE::AIAssistant
 		// Update the global locale state.
 		void UpdateGlobalLocale(const FString& LocaleString);
 
+		// Opaque handle to data needed to unregister conversation update callback.
+		using FConversationUpdateHandle = TTuple<FString, FString>;
+
 		// Register for conversation updates. Return true to unregister, or explicitly call
 		// UnregisterOnConversationUpdate.
-		FString RegisterOnConversationUpdate(
+		TFuture<TValueOrError<FConversationUpdateHandle, FString>> RegisterOnConversationUpdate(
 			TFunction<bool(FConversationId)> OnConversationUpdate);
 
-		// Unregister conversation update callback using the HandlerId returned from
+		// Unregister conversation update callback using the CallbackRegistration returned from
 		// RegisterOnConversationUpdate.
-		void UnregisterOnConversationUpdate(const FString& HandlerId);
+		void UnregisterOnConversationUpdate(const FConversationUpdateHandle& Registration);
 
 	protected:
 		// Format a function call of a member with result handling.
@@ -273,8 +276,9 @@ namespace UE::AIAssistant
 			const TCHAR* Arguments = TEXT(""));
 		
 		// Register a callback with Javascript, allowing the given TFunction to be called multiple
-		// times from Javascript. To unregister, 
-		FString RegisterCallback(
+		// times from Javascript. This is virtual so that the returned handler ID can be captured
+		// in tests.
+		virtual FString RegisterCallback(
 			TFunction<bool(UAIAssistantWebJavaScriptResultDelegate::FResult)> Callback);
 
 		// Execute a javascript function handling the result with the specified handler.
@@ -398,4 +402,4 @@ namespace UE::AIAssistant
 		static const FString FunctionCallFormatTemplate;
 		static const FString FunctionCallbackFormatTemplate;
 	};
-}
+}
\ No newline at end of file
diff --git a/Engine/Plugins/Experimental/AIAssistant/Source/AIAssistant/Private/Tests/AIAssistantFakeWebApi.cpp b/Engine/Plugins/Experimental/AIAssistant/Source/AIAssistant/Private/Tests/AIAssistantFakeWebApi.cpp
--- a/Engine/Plugins/Experimental/AIAssistant/Source/AIAssistant/Private/Tests/AIAssistantFakeWebApi.cpp
+++ b/Engine/Plugins/Experimental/AIAssistant/Source/AIAssistant/Private/Tests/AIAssistantFakeWebApi.cpp
@@ -87,6 +87,13 @@ namespace UE::AIAssistant
 			*WebJavaScriptResultDelegate, HandlerId, ResultJson, bResultJsonIsError);
 	}
 
+	FString FFakeWebApi::RegisterCallback(
+		TFunction<bool(UAIAssistantWebJavaScriptResultDelegate::FResult)> Callback)
+	{
+		RegisteredCallbackHandlerIds.Add(FWebApi::RegisterCallback(MoveTemp(Callback)));
+		return RegisteredCallbackHandlerIds.Last();
+	}
+
 	void FFakeWebApi::ExecuteAsyncFunction(
 		const TCHAR* InstanceName, const TCHAR* FunctionName, const TCHAR* Arguments,
 		const TCHAR* HandlerId)
diff --git a/Engine/Plugins/Experimental/AIAssistant/Source/AIAssistant/Private/Tests/AIAssistantFakeWebApi.h b/Engine/Plugins/Experimental/AIAssistant/Source/AIAssistant/Private/Tests/AIAssistantFakeWebApi.h
--- a/Engine/Plugins/Experimental/AIAssistant/Source/AIAssistant/Private/Tests/AIAssistantFakeWebApi.h
+++ b/Engine/Plugins/Experimental/AIAssistant/Source/AIAssistant/Private/Tests/AIAssistantFakeWebApi.h
@@ -117,6 +117,9 @@ namespace UE::AIAssistant
 			const TCHAR* InstanceName, const TCHAR* FunctionName, const TCHAR* Arguments,
 			const TCHAR* HandlerId) override;
 
+		FString RegisterCallback(
+			TFunction<bool(UAIAssistantWebJavaScriptResultDelegate::FResult)> Callback) override;
+
 		// If possible, get the value of a result as JSON.  If the value cannot be converted to
 		// JSON return an empty string.
 		template<typename ValueType>
@@ -135,6 +138,7 @@ namespace UE::AIAssistant
 		}
 
 	public:
+		TArray<FString> RegisteredCallbackHandlerIds;
 		TArray<FExecutedAsyncFunction> ExecutedAsyncFunctions;
 
 	public:
diff --git a/Engine/Plugins/Experimental/AIAssistant/Source/AIAssistant/Private/Tests/AIAssistantWebApiTest.cpp b/Engine/Plugins/Experimental/AIAssistant/Source/AIAssistant/Private/Tests/AIAssistantWebApiTest.cpp
--- a/Engine/Plugins/Experimental/AIAssistant/Source/AIAssistant/Private/Tests/AIAssistantWebApiTest.cpp
+++ b/Engine/Plugins/Experimental/AIAssistant/Source/AIAssistant/Private/Tests/AIAssistantWebApiTest.cpp
@@ -52,7 +52,7 @@ namespace UE::AIAssistant
 		{
 			return WebApi.FormatResultAndErrorHandlers(HandlerId);
 		}
-		
+
 		static UAIAssistantWebJavaScriptResultDelegate& GetJavaScriptResultDelegate(
 			FWebApi& WebApi)
 		{
@@ -374,18 +374,21 @@ bool FAIAssistantWebApiTestRegisterOnConversationUpdate::RunTest(const FString&
 {
 	FFakeWebApiContainer WebApi;
 	int CallCount = 0;
-	FString HandlerId = WebApi->RegisterOnConversationUpdate([&CallCount](FConversationId ConversationId)
+	auto Future =
+		WebApi->RegisterOnConversationUpdate([&CallCount](FConversationId ConversationId)
 		{
 			CallCount += 1;
 			return false;
 		});
-	
+	const FString& HandlerId = WebApi->RegisteredCallbackHandlerIds.Last();
 	return
-		WebApi->TestExpectAsyncFunctionCall(
+		WebApi->TestExpectAsyncFunctionCallAndComplete(
 			*this,
 			nullptr,
-			TEXT("RegisterOnConversationUpdate"),
-			*FWebApiAccessor::FormatFunctionCallback(*WebApi, *HandlerId)).IsValid()
+			TEXT("registerOnConversationUpdate"),
+			*FWebApiAccessor::FormatFunctionCallback(*WebApi, *HandlerId),
+			TEXT(""),
+			false)
 		&& TestEqual(TEXT("RegisterOnConversationUpdate_NoCalls"), CallCount, 0);
 }
 
@@ -400,26 +403,31 @@ bool FAIAssistantWebApiTestRegisterOnConversationUpdateAndInvoke::RunTest(
 	FFakeWebApiContainer WebApi;
 	FString ActualConversationId;
 	int CallCount = 0;
-	FString CallbackHandle = WebApi->RegisterOnConversationUpdate(
+	auto Result = WebApi->RegisterOnConversationUpdate(
 		[&ActualConversationId, &CallCount](FConversationId ConversationId)
 		{
 			ActualConversationId = ConversationId.Id;
 			CallCount += 1;
 			return false;
 		});
+	const FString& HandlerId = WebApi->RegisteredCallbackHandlerIds.Last();
 	if (!WebApi->TestExpectAsyncFunctionCallAndComplete(
 		*this,
 		nullptr,
 		TEXT("registerOnConversationUpdate"),
-		*FWebApiAccessor::FormatFunctionCallback(*WebApi, *CallbackHandle),
+		*FWebApiAccessor::FormatFunctionCallback(*WebApi, *HandlerId),
 		TEXT(""),
 		false))
 	{
 		return false;
 	}
+
+	auto& Registration = Result.Get().GetValue();
+	auto& CallbackHandle = Registration.Get<1>();
+
 	FString ExpectedConversationId = TEXT("foobar");
 	WebApi->InvokeCallback(
-		CallbackHandle,
+		HandlerId,
 		*FString::Printf(TEXT(R"js({ "id": "%s" })js"), *ExpectedConversationId),
 		false);
 	return
@@ -441,7 +449,7 @@ bool FAIAssistantWebApiTestRegisterOnConversationUpdateAndInvokeTwice::RunTest(
 	FFakeWebApiContainer WebApi;
 	TArray<FString> ActualConversationIds;
 	int CallCount = 0;
-	FString CallbackHandle = WebApi->RegisterOnConversationUpdate(
+	auto Result = WebApi->RegisterOnConversationUpdate(
 		[&ActualConversationIds, &CallCount](FConversationId ConversationId)
 		{
 			ActualConversationIds.Push(ConversationId.Id);
@@ -449,21 +457,27 @@ bool FAIAssistantWebApiTestRegisterOnConversationUpdateAndInvokeTwice::RunTest(
 			return false;
 		});
 
+	const FString& HandlerId = WebApi->RegisteredCallbackHandlerIds.Last();
+
 	if (!WebApi->TestExpectAsyncFunctionCallAndComplete(
 		*this,
 		nullptr,
-		TEXT("RegisterOnConversationUpdate"),
-		*FWebApiAccessor::FormatFunctionCallback(*WebApi, *CallbackHandle),
+		TEXT("registerOnConversationUpdate"),
+		*FWebApiAccessor::FormatFunctionCallback(*WebApi, *HandlerId),
 		TEXT(""),
 		false))
 	{
 		return false;
 	}
+
+	auto& Registration = Result.Get().GetValue();
+	auto& CallbackHandle = Registration.Get<1>();
+
 	TArray<FString> ExpectedConversationIds = { TEXT("foo"), TEXT("bar") };
 	for (const auto& ConversationId : ExpectedConversationIds)
 	{
 		WebApi->InvokeCallback(
-			CallbackHandle,
+			HandlerId,
 			*FString::Printf(TEXT(R"js({ "id": "%s" })js"), *ConversationId),
 			false);
 	}
@@ -487,29 +501,33 @@ bool FAIAssistantWebApiTestUnegisterOnConversationUpdateViaReturnValue::RunTest(
 	FFakeWebApiContainer WebApi;
 	TArray<FString> ActualConversationIds;
 	int CallCount = 0;
-	FString CallbackHandle = WebApi->RegisterOnConversationUpdate(
+	auto Result = WebApi->RegisterOnConversationUpdate(
 		[&ActualConversationIds, &CallCount](FConversationId ConversationId)
 		{
 			ActualConversationIds.Push(ConversationId.Id);
 			CallCount += 1;
 			return true;
 		});
-
+	const FString& HandlerId = WebApi->RegisteredCallbackHandlerIds.Last();
 	if (!WebApi->TestExpectAsyncFunctionCallAndComplete(
 		*this,
 		nullptr,
-		TEXT("RegisterOnConversationUpdate"),
-		*FWebApiAccessor::FormatFunctionCallback(*WebApi, *CallbackHandle),
+		TEXT("registerOnConversationUpdate"),
+		*FWebApiAccessor::FormatFunctionCallback(*WebApi, *HandlerId),
 		TEXT(""),
 		false))
 	{
 		return false;
 	}
+
+	auto& Registration = Result.Get().GetValue();
+	auto& CallbackHandle = Registration.Get<1>();
+
 	TArray<FString> ExpectedConversationIds = { TEXT("foo") };
 	for (const auto& ConversationId : ExpectedConversationIds)
 	{
 		WebApi->InvokeCallback(
-			CallbackHandle,
+			*HandlerId,
 			*FString::Printf(TEXT(R"js({ "id": "%s" })js"), *ConversationId),
 			false);
 	}
@@ -528,11 +546,61 @@ IMPLEMENT_SIMPLE_AUTOMATION_TEST(
 bool FAIAssistantWebApiTestUnregisterOnConversationUpdate::RunTest(const FString& UnusedParameters)
 {
 	FFakeWebApiContainer WebApi;
-	FString FakeHandlerId = TEXT("foobar");
-	WebApi->UnregisterOnConversationUpdate(FakeHandlerId);
-	return WebApi->TestExpectAsyncFunctionCall(
+
+	// Register a callback.
+	TArray<FString> ActualConversationIds;
+	int CallCount = 0;
+	auto Result = WebApi->RegisterOnConversationUpdate(
+		[&ActualConversationIds, &CallCount](FConversationId ConversationId)
+		{
+			ActualConversationIds.Push(ConversationId.Id);
+			CallCount += 1;
+			return false;
+		});
+	const FString& HandlerId = WebApi->RegisteredCallbackHandlerIds.Last();
+	if (!WebApi->TestExpectAsyncFunctionCallAndComplete(
+		*this,
+		nullptr,
+		TEXT("registerOnConversationUpdate"),
+		*FWebApiAccessor::FormatFunctionCallback(*WebApi, *HandlerId),
+		TEXT(""),
+		false))
+	{
+		return false;
+	}
+
+	auto& Registration = Result.Get().GetValue();
+	auto& CallbackHandle = Registration.Get<1>();
+
+	// Invoke callback once.
+	WebApi->InvokeCallback(
+		HandlerId,
+		*FString::Printf(TEXT(R"js({ "id": "%s" })js"), TEXT("foo")),
+		false);
+
+	// Unregister the callback.
+	WebApi->UnregisterOnConversationUpdate(Registration);
+	if (!WebApi->TestExpectAsyncFunctionCall(
 		*this, nullptr, TEXT("unregisterOnConversationUpdate"),
-		*FString::Printf(TEXT("\"%s\""), *FakeHandlerId)).IsValid();
+		*FString::Printf(TEXT("%s"), *CallbackHandle)).IsValid())
+	{
+		return false;
+	}
+
+	// Try to invoke callback again.
+	WebApi->InvokeCallback(
+		HandlerId,
+		*FString::Printf(TEXT(R"js({ "id": "%s" })js"), TEXT("bar")),
+		false);
+
+	// Ensure callback was only called the first time.
+	TArray<FString> ExpectedConversationIds = { TEXT("foo") };
+	return
+		TestEqual(
+			TEXT("UnregisterOnConversationUpdate_IdsMatch"),
+			ActualConversationIds,
+			ExpectedConversationIds)
+		&& TestEqual(TEXT("UnregisterOnConversationUpdate_CallCount"), CallCount, 1);
 }
 
 #endif  // WITH_DEV_AUTOMATION_TESTS
\ No newline at end of file
