[AIAssistant] Refactor FFakeToolsetHandler to a common location

#Synced-CL 48676371

diff --git a/Engine/Plugins/Experimental/AIAssistant/Source/AIAssistant/Private/Tests/AIAssistantFakeToolsetHandler.cpp b/Engine/Plugins/Experimental/AIAssistant/Source/AIAssistant/Private/Tests/AIAssistantFakeToolsetHandler.cpp
new file mode 100644
--- /dev/null
+++ b/Engine/Plugins/Experimental/AIAssistant/Source/AIAssistant/Private/Tests/AIAssistantFakeToolsetHandler.cpp
@@ -0,0 +1,45 @@
+// Copyright Epic Games, Inc. All Rights Reserved.
+
+#include "AIAssistantFakeToolsetHandler.h"
+
+#if WITH_DEV_AUTOMATION_TESTS
+
+namespace UE::AIAssistant
+{
+	FFakeToolset::FFakeToolset(const FString& Name) :
+		Name(Name)
+	{
+		ToolPromises.Add(
+			TEXT("SomeTool"),
+			MakeFulfilledPromise<TValueOrError<FString, FString>>(MakeValue(SUCCESSFUL_RESULT)));
+	}
+
+	TFuture<TValueOrError<FString, FString>> FFakeToolset::ExecuteTool(
+		const FString& ToolName, const FString& JsonInput)
+	{
+		if (!ToolPromises.Contains(ToolName))
+		{
+			return MakeFulfilledPromise<TValueOrError<FString, FString>>(
+				MakeError(FString(TEXT("Tool not found")))).GetFuture();
+		}
+
+		return ToolPromises[ToolName].GetFuture();
+	}
+
+	FString FFakeToolset::GetToolsetJsonSchema() const
+	{
+		return FString();
+	}
+
+	FString FFakeToolset::GetToolsetName() const
+	{
+		return Name;
+	}
+
+	FString FFakeToolset::GetToolsetVersion() const
+	{
+		return TEXT("1.0");
+	}
+}
+
+#endif
\ No newline at end of file
diff --git a/Engine/Plugins/Experimental/AIAssistant/Source/AIAssistant/Private/Tests/AIAssistantFakeToolsetHandler.h b/Engine/Plugins/Experimental/AIAssistant/Source/AIAssistant/Private/Tests/AIAssistantFakeToolsetHandler.h
new file mode 100644
--- /dev/null
+++ b/Engine/Plugins/Experimental/AIAssistant/Source/AIAssistant/Private/Tests/AIAssistantFakeToolsetHandler.h
@@ -0,0 +1,40 @@
+// Copyright Epic Games, Inc. All Rights Reserved.
+#pragma once
+
+#include "AIAssistantToolsetRegistry.h"
+
+#include "Misc/AutomationTest.h"
+#include "UObject/StrongObjectPtr.h"
+
+#include "AIAssistantTestFlags.h"
+#include "AIAssistantToolsetHandler.h"
+
+#if WITH_DEV_AUTOMATION_TESTS
+
+namespace UE::AIAssistant
+{
+
+class FFakeToolset : public IToolsetHandler
+{
+public:
+	static const inline FString SUCCESSFUL_RESULT = FString(TEXT(R"({"result": "success"})"));
+
+	FFakeToolset(const FString& Name);
+
+	TFuture<TValueOrError<FString, FString>> ExecuteTool(
+		const FString& ToolName, const FString& JsonInput) override;
+
+	FString GetToolsetJsonSchema() const override;
+
+	FString GetToolsetName() const override;
+
+	FString GetToolsetVersion() const override;
+
+	FString Name;
+
+	TMap<FString, TPromise<TValueOrError<FString, FString>>> ToolPromises;
+};
+
+}
+
+#endif
diff --git a/Engine/Plugins/Experimental/AIAssistant/Source/AIAssistant/Private/Tests/AIAssistantToolsetRegistryTest.cpp b/Engine/Plugins/Experimental/AIAssistant/Source/AIAssistant/Private/Tests/AIAssistantToolsetRegistryTest.cpp
--- a/Engine/Plugins/Experimental/AIAssistant/Source/AIAssistant/Private/Tests/AIAssistantToolsetRegistryTest.cpp
+++ b/Engine/Plugins/Experimental/AIAssistant/Source/AIAssistant/Private/Tests/AIAssistantToolsetRegistryTest.cpp
@@ -7,6 +7,7 @@
 
 #include "AIAssistantTestFlags.h"
 #include "AIAssistantToolsetHandler.h"
+#include "AIAssistantFakeToolsetHandler.h"
 
 #if WITH_DEV_AUTOMATION_TESTS
 
@@ -110,50 +111,6 @@ bool FAIAssistantToolsetDescriptorFromStringMultipleDots::RunTest(const FString&
 	return true;
 }
 
-class FakeToolset : public IToolsetHandler
-{
-public:
-	static const inline FString SUCCESSFUL_RESULT = FString(TEXT(R"({"result": "success"})"));
-
-	TFuture<TValueOrError<FString, FString>> ExecuteTool(
-		const FString& ToolName, const FString& JsonInput) override
-	{
-		if (ToolName != TEXT("SomeTool"))
-		{
-			return MakeFulfilledPromise<TValueOrError<FString, FString>>(
-				MakeError(FString(TEXT("Tool not found")))).GetFuture();
-		}
-
-		FString Result(SUCCESSFUL_RESULT);
-		return MakeFulfilledPromise<TValueOrError<FString, FString>>(
-			MakeValue(Result)).GetFuture();
-	}
-
-	FString GetToolsetJsonSchema() const override
-	{
-		return FString();
-	}
-
-	FString GetToolsetName() const override
-	{
-		return TEXT("Fake");
-	}
-
-	FString GetToolsetVersion() const override
-	{
-		return TEXT("1.0");
-	}
-};
-
-class AnotherFakeToolset : public FakeToolset
-{
-public:
-	FString GetToolsetName() const override
-	{
-		return TEXT("AnotherFake");
-	}
-};
-
 IMPLEMENT_SIMPLE_AUTOMATION_TEST(
 	FAIAssistantToolsetRegistryToolsetHandlerRegistrationSuccess,
 	"AI.Assistant.ToolsetRegistry.Registration.Success",
@@ -165,7 +122,7 @@ bool FAIAssistantToolsetRegistryToolsetHandlerRegistrationSuccess::RunTest(
 	TUniquePtr<FToolsetRegistry> Registry = MakeUnique<FToolsetRegistry>();
 
 	// Test whether a valid toolset handler is accepted.
-	TSharedPtr<IToolsetHandler> ValidToolsetHandler = MakeShared<FakeToolset>();
+	TSharedPtr<IToolsetHandler> ValidToolsetHandler = MakeShared<FFakeToolset>(TEXT("Fake"));
 	(void)TestTrue(
 		TEXT("RegisterToolsetHandler_Valid"),
 		Registry->RegisterToolset(ValidToolsetHandler));
@@ -182,7 +139,7 @@ bool FAIAssistantToolsetRegistryToolsetHandlerRegistrationDoubleRegistration::Ru
 	const FString& UnusedParameters)
 {
 	TUniquePtr<FToolsetRegistry> Registry = MakeUnique<FToolsetRegistry>();
-	TSharedPtr<IToolsetHandler> ValidToolsetHandler = MakeShared<FakeToolset>();
+	TSharedPtr<IToolsetHandler> ValidToolsetHandler = MakeShared<FFakeToolset>(TEXT("Fake"));
 	Registry->RegisterToolset(ValidToolsetHandler);
 
 	// Test whether duplicate registration is rejected.
@@ -203,7 +160,7 @@ bool FAIAssistantToolsetRegistryToolsetHandlerUnregistrationSuccess::RunTest(
 	const FString& UnusedParameters)
 {
 	TUniquePtr<FToolsetRegistry> Registry = MakeUnique<FToolsetRegistry>();
-	TSharedPtr<IToolsetHandler> ValidToolsetHandler = MakeShared<FakeToolset>();
+	TSharedPtr<IToolsetHandler> ValidToolsetHandler = MakeShared<FFakeToolset>(TEXT("Fake"));
 	Registry->RegisterToolset(ValidToolsetHandler);
 
 	// Test whether unregistration of a registered toolset handler is accepted.
@@ -223,10 +180,11 @@ bool FAIAssistantToolsetRegistryToolsetHandlerUnregistrationUnregisteredToolset:
 	const FString& UnusedParameters)
 {
 	TUniquePtr<FToolsetRegistry> Registry = MakeUnique<FToolsetRegistry>();
-	Registry->RegisterToolset(MakeShared<FakeToolset>());
+	Registry->RegisterToolset(MakeShared<FFakeToolset>(TEXT("Fake")));
 
 	// Test whether unregistration of an non-registered toolset handler is rejected.
-	TSharedPtr<IToolsetHandler> UnregisteredToolsetHandler = MakeShared<AnotherFakeToolset>();
+	TSharedPtr<IToolsetHandler> UnregisteredToolsetHandler =
+		MakeShared<FFakeToolset>(TEXT("AnotherFake"));
 	AddExpectedError(TEXT("Toolset 'AnotherFake' not registered"));
 	(void)TestFalse(
 		TEXT("UnregisterToolsetHandler_NotRegistered"),
@@ -244,7 +202,7 @@ bool FAIAssistantToolsetRegistryToolsetHandlerUnregistrationDoubleUnregistration
 	const FString& UnusedParameters)
 {
 	TUniquePtr<FToolsetRegistry> Registry = MakeUnique<FToolsetRegistry>();
-	TSharedPtr<IToolsetHandler> ValidToolsetHandler = MakeShared<FakeToolset>();
+	TSharedPtr<IToolsetHandler> ValidToolsetHandler = MakeShared<FFakeToolset>(TEXT("Fake"));
 	Registry->RegisterToolset(ValidToolsetHandler);
 	Registry->UnregisterToolset(ValidToolsetHandler);
 
@@ -265,8 +223,8 @@ IMPLEMENT_SIMPLE_AUTOMATION_TEST(
 bool FAIAssistantToolsetRegistryToolsetHandlerFindSuccess::RunTest(const FString& UnusedParameters)
 {
 	TUniquePtr<FToolsetRegistry> Registry = MakeUnique<FToolsetRegistry>();
-	TSharedPtr<IToolsetHandler> Fake = MakeShared<FakeToolset>();
-	TSharedPtr<IToolsetHandler> AnotherFake = MakeShared<AnotherFakeToolset>();
+	TSharedPtr<IToolsetHandler> Fake = MakeShared<FFakeToolset>(TEXT("Fake"));
+	TSharedPtr<IToolsetHandler> AnotherFake = MakeShared<FFakeToolset>(TEXT("AnotherFake"));
 	Registry->RegisterToolset(Fake);
 	Registry->RegisterToolset(AnotherFake);
 
@@ -300,7 +258,7 @@ bool FAIAssistantToolsetRegistryToolsetHandlerFindNonexistentToolsetLogError::Ru
 	const FString& UnusedParameters)
 {
 	TUniquePtr<FToolsetRegistry> Registry = MakeUnique<FToolsetRegistry>();
-	TSharedPtr<IToolsetHandler> Fake = MakeShared<FakeToolset>();
+	TSharedPtr<IToolsetHandler> Fake = MakeShared<FFakeToolset>(TEXT("Fake"));
 	Registry->RegisterToolset(Fake);
 	
 	// Test finding an unregistered toolset handler.
@@ -321,7 +279,7 @@ bool FAIAssistantToolsetRegistryToolsetHandlerFindNonexistentToolsetReturnErrorM
 	const FString& UnusedParameters)
 {
 	TUniquePtr<FToolsetRegistry> Registry = MakeUnique<FToolsetRegistry>();
-	TSharedPtr<IToolsetHandler> Fake = MakeShared<FakeToolset>();
+	TSharedPtr<IToolsetHandler> Fake = MakeShared<FFakeToolset>(TEXT("Fake"));
 	Registry->RegisterToolset(Fake);
 
 	// Test finding an unregistered toolset handler with error message.
@@ -350,7 +308,7 @@ const FString& UnusedParameters)
 	TUniquePtr<FToolsetRegistry> Registry = MakeUnique<FToolsetRegistry>();
 
 	// Register a valid toolset handler.
-	TSharedPtr<IToolsetHandler> ValidToolsetHandler = MakeShared<FakeToolset>();
+	TSharedPtr<IToolsetHandler> ValidToolsetHandler = MakeShared<FFakeToolset>(TEXT("Fake"));
 	(void)TestTrue(
 		TEXT("RegisterToolsetHandler_Valid"),
 		Registry->RegisterToolset(ValidToolsetHandler));
@@ -365,7 +323,7 @@ const FString& UnusedParameters)
 			(void)TestEqual(
 				TEXT("ExecuteTool_ResultJson"),
 				Result.StealValue(),
-				FakeToolset::SUCCESSFUL_RESULT);
+				FFakeToolset::SUCCESSFUL_RESULT);
 		})
 		.Wait();
 
@@ -383,7 +341,7 @@ bool FAIAssistantToolsetRegistryToolsetHandlerExecuteToolNonexistantTool::RunTes
 	TUniquePtr<FToolsetRegistry> Registry = MakeUnique<FToolsetRegistry>();
 
 	// Register a valid toolset handler.
-	TSharedPtr<IToolsetHandler> ValidToolsetHandler = MakeShared<FakeToolset>();
+	TSharedPtr<IToolsetHandler> ValidToolsetHandler = MakeShared<FFakeToolset>(TEXT("Fake"));
 	(void)TestTrue(
 		TEXT("RegisterToolsetHandler_Valid"),
 		Registry->RegisterToolset(ValidToolsetHandler));
@@ -414,7 +372,7 @@ bool FAIAssistantToolsetRegistryToolsetHandlerExecuteToolNonexistantToolset::Run
 	TUniquePtr<FToolsetRegistry> Registry = MakeUnique<FToolsetRegistry>();
 
 	// Register a valid toolset handler.
-	TSharedPtr<IToolsetHandler> ValidToolsetHandler = MakeShared<FakeToolset>();
+	TSharedPtr<IToolsetHandler> ValidToolsetHandler = MakeShared<FFakeToolset>(TEXT("Fake"));
 	(void)TestTrue(
 		TEXT("RegisterToolsetHandler_Valid"),
 		Registry->RegisterToolset(ValidToolsetHandler));
@@ -446,7 +404,7 @@ bool FAIAssistantToolsetRegistryToolsetHandlerExecuteToolBadToolnameFormat::RunT
 	TUniquePtr<FToolsetRegistry> Registry = MakeUnique<FToolsetRegistry>();
 
 	// Register a valid toolset handler.
-	TSharedPtr<IToolsetHandler> ValidToolsetHandler = MakeShared<FakeToolset>();
+	TSharedPtr<IToolsetHandler> ValidToolsetHandler = MakeShared<FFakeToolset>(TEXT("Fake"));
 	(void)TestTrue(
 		TEXT("RegisterToolsetHandler_Valid"),
 		Registry->RegisterToolset(ValidToolsetHandler));
