[AIAssistant] Add serialization structs to support tool calling

#Synced-CL 48821063

diff --git a/Engine/Plugins/Experimental/AIAssistant/Source/AIAssistant/Private/AIAssistantWebApi.h b/Engine/Plugins/Experimental/AIAssistant/Source/AIAssistant/Private/AIAssistantWebApi.h
--- a/Engine/Plugins/Experimental/AIAssistant/Source/AIAssistant/Private/AIAssistantWebApi.h
+++ b/Engine/Plugins/Experimental/AIAssistant/Source/AIAssistant/Private/AIAssistantWebApi.h
@@ -43,56 +43,189 @@ namespace UE::AIAssistant
 	enum class EMessageContentType
 	{
 		Text,
+		ToolCall,
+		ToolResponse,
 	};
 
 	#define UE_AI_ASSISTANT_MESSAGE_CONTENT_TYPE_ENUM(X) \
-		X(EMessageContentType::Text, "text")
+		X(EMessageContentType::Text, "text"), \
+		X(EMessageContentType::ToolCall, "tool_call"), \
+		X(EMessageContentType::ToolResponse, "tool_response")
 
 	UE_ENUM_METADATA_DECLARE(EMessageContentType, UE_AI_ASSISTANT_MESSAGE_CONTENT_TYPE_ENUM);
 
 	// Text message content.
 	struct FTextMessageContent : public FJsonSerializable
 	{
+		// Text message in a conversation.
 		FString Text;
+		// (Optional) Text format, default is set as markdown.
+		TOptional<FString> Format;
+
+		bool operator==(const FTextMessageContent& Other) const {
+			static const FString Markdown = TEXT("markdown");
+
+			const FString& ThisFormat = Format.IsSet() ? Format.GetValue() : Markdown;
+			const FString& OtherFormat = Other.Format.IsSet() ? Other.Format.GetValue() : Markdown;
+
+			return Text == Other.Text && ThisFormat == OtherFormat;
+		}
 
 		BEGIN_JSON_SERIALIZER
 			JSON_SERIALIZE("text", Text);
+			JSON_SERIALIZE_OPTIONAL("format", Format);
+		END_JSON_SERIALIZER
+	};
+
+	// Tool call message content.
+	struct FToolCallContent : public FJsonSerializable
+	{
+		// Tool name (e.g., 'agent_status', 'reference').
+		FString Name;
+		// Tool-specific arguments.
+		FString JsonArguments;
+		// UUID for tool calls that expect response.
+		TOptional<FString> ToolCallId;
+		// Whether agent expects response, default is set to false.
+		TOptional<bool> ResponseRequired;
+
+		bool operator==(const FToolCallContent& Other) const {
+			return Name == Other.Name
+				&& JsonArguments == Other.JsonArguments
+				&& ToolCallId == Other.ToolCallId
+				&& ResponseRequired == Other.ResponseRequired;
+		};
+
+		BEGIN_JSON_SERIALIZER
+			JSON_SERIALIZE("name", Name);
+			JSON_SERIALIZE("arguments", JsonArguments);
+			JSON_SERIALIZE_OPTIONAL("toolCallId", ToolCallId);
+			JSON_SERIALIZE_OPTIONAL("responseRequired", ResponseRequired);
+		END_JSON_SERIALIZER
+	};
+	
+	// Tool response content.
+	struct FToolResponseContent : public FJsonSerializable
+	{
+		// ID of the tool call this responds to.
+		FString ToolCallId;
+		// Name of the tool that was called.
+		FString Name;
+		// Response data from client execution.
+		FString JsonResponse;
+		// Whether execution succeeded. An empty value represents true.
+		TOptional<bool> Success;
+		// Error details if success=false.
+		TOptional<FString> ErrorMessage;
+
+		bool operator==(const FToolResponseContent& Other) const {
+			bool SuccessValue = !Success.IsSet() || *Success;
+			bool OtherSuccessValue = !Other.Success.IsSet() || *Other.Success;
+
+			return ToolCallId == Other.ToolCallId
+				&& Name == Other.Name
+				&& JsonResponse == Other.JsonResponse
+				&& SuccessValue == OtherSuccessValue
+				&& ErrorMessage == Other.ErrorMessage;
+		};
+
+		BEGIN_JSON_SERIALIZER
+			JSON_SERIALIZE("toolCallId", ToolCallId);
+			JSON_SERIALIZE("name", Name);
+			JSON_SERIALIZE("response", JsonResponse);
+			JSON_SERIALIZE_OPTIONAL("success", Success);
+			JSON_SERIALIZE_OPTIONAL("errorMessage", ErrorMessage);
 		END_JSON_SERIALIZER
 	};
 
+	using FMessageContentVariant = TVariant<
+		FTextMessageContent,
+		FToolCallContent,
+		FToolResponseContent
+	>;
+
 	// Content of a message.
 	struct FMessageContent : public FJsonSerializable
 	{
 		// Type of the messageContent field.
 		EMessageContentType ContentType;
 		// Content of the message, this can be extended in future to support additional types of content.
-		TVariant<FTextMessageContent> Content;
+		FMessageContentVariant Content;
 		// Whether the message is visible to the user.
 		bool bVisibleToUser = true;
 
+		bool operator==(const FMessageContent& Other) const
+		{
+			return (
+				ContentType == Other.ContentType
+				&& bVisibleToUser == Other.bVisibleToUser
+				&& Visit(
+					[&](const auto& Value, const auto& Other)
+					{
+						if constexpr (std::is_same_v<decltype(Value), decltype(Other)>)
+						{
+							return Value == Other;
+						}
+						else
+						{
+							return false;
+						}
+					}, Content, Other.Content));
+		}
+
 		BEGIN_JSON_SERIALIZER
 			UE_JSON_SERIALIZE_ENUM_VARIANT_BEGIN("contentType", ContentType, "content", Content);
 				UE_JSON_SERIALIZE_ENUM_VARIANT(EMessageContentType::Text, FTextMessageContent);
+				UE_JSON_SERIALIZE_ENUM_VARIANT(EMessageContentType::ToolCall, FToolCallContent);
+				UE_JSON_SERIALIZE_ENUM_VARIANT(EMessageContentType::ToolResponse, FToolResponseContent);
 			UE_JSON_SERIALIZE_ENUM_VARIANT_END();
 			JSON_SERIALIZE_WITHDEFAULT("visibleToUser", bVisibleToUser, true);
 		END_JSON_SERIALIZER
 
 	};
 
+	// ID of a message, generated by the assistant backend.
+	struct FMessageId : public FJsonSerializable
+	{
+		// Unique ID of an object.
+		FString Id;
+		// Discriminator of the type.
+		FString Type;
+
+		bool operator==(const FMessageId& Other) const
+		{
+			return Id == Other.Id && Type == Other.Type;
+		}
+
+		BEGIN_JSON_SERIALIZER
+			JSON_SERIALIZE("id", Id);
+			JSON_SERIALIZE("type", Type);
+		END_JSON_SERIALIZER
+	};
+
 	// Message within a conversation.
 	struct FMessage : public FJsonSerializable
 	{
+		// Populated on the server side and ignore when specified on the client.
+		TOptional<FMessageId> MessageId;
 		// When the message was created (UTC). This is populated on the server side and ignored when
 		// specified on the client side.
 		TOptional<FDateTime> Date;
-
 		// Role or source of the message.
 		EMessageRole MessageRole;
-
 		// Content of the message.
 		TArray<FMessageContent> MessageContent;
 
+		bool operator==(const FMessage& Other) const
+		{
+			return MessageId == Other.MessageId
+				&& Date == Other.Date
+				&& MessageRole == Other.MessageRole
+				&& MessageContent == Other.MessageContent;
+		}
+
 		BEGIN_JSON_SERIALIZER
+			JSON_SERIALIZE_OPTIONAL_OBJECT_SERIALIZABLE("id", MessageId);
 			JSON_SERIALIZE_OPTIONAL("date", Date);
 			JSON_SERIALIZE_ENUM("messageRole", MessageRole);
 			JSON_SERIALIZE_ARRAY_SERIALIZABLE("messageContent", MessageContent, FMessageContent);
@@ -105,11 +238,75 @@ namespace UE::AIAssistant
 		// Unique ID of the conversation.
 		FString Id;
 
+		bool operator==(const FConversationId& Other) const
+		{
+			return Id == Other.Id;
+		}
+
 		BEGIN_JSON_SERIALIZER
 			JSON_SERIALIZE("id", Id);
 		END_JSON_SERIALIZER
 	};
 
+	// Metadata used to index and display a conversation.
+	struct FConversationDescriptor : public FJsonSerializable {
+		// ID of the conversation, this is generated on the server side
+		// and ignored if set by the client.
+		TOptional<FConversationId> ConversationId;
+		// Display name / one-line description of the conversation that
+		// is simply used to help a user find it within a flat list of
+		// all conversations.
+		FString DisplayName;
+		// Labels applied to the conversation to make it possible for a user
+		// to organize and search for conversations. For example:
+		// "epic.favorite" could be reserved for favorites (starred)
+		// conversations.
+		TArray<FString> Labels;
+		// When the conversation was created (UTC milliseconds since 
+		// the epoch). This is derived from the first message in the 
+		// conversation.
+		FDateTime CreatedAt;
+		// When the conversation was last updated (UTC milliseconds since 
+		// the epoch). This is derived  from the most recent message in the 
+		// conversation.
+		FDateTime UpdatedAt;
+
+		bool operator==(const FConversationDescriptor& Other) const
+		{
+			return ConversationId == Other.ConversationId
+				&& DisplayName == Other.DisplayName
+				&& Labels == Other.Labels
+				&& CreatedAt == Other.CreatedAt
+				&& UpdatedAt == Other.UpdatedAt;
+		}
+
+		BEGIN_JSON_SERIALIZER
+			JSON_SERIALIZE_OPTIONAL_OBJECT_SERIALIZABLE("conversationId", ConversationId);
+			JSON_SERIALIZE("displayName", DisplayName);
+			JSON_SERIALIZE_ARRAY("labels", Labels);
+			JSON_SERIALIZE("createdAt", CreatedAt);
+			JSON_SERIALIZE("updatedAt", UpdatedAt);
+		END_JSON_SERIALIZER
+	};
+
+	// Conversation.
+	struct FConversation : public FJsonSerializable
+	{
+		// Metadata used to index and display a conversation.
+		FConversationDescriptor Descriptor;
+		// Messages in the conversation, ordered chronologically.
+		TArray<FMessage> Messages;
+
+		bool operator==(const FConversation& Other) const
+		{
+			return Descriptor == Other.Descriptor && Messages == Other.Messages;
+		}
+
+		BEGIN_JSON_SERIALIZER
+			JSON_SERIALIZE_OBJECT_SERIALIZABLE("descriptor", Descriptor);
+			JSON_SERIALIZE_ARRAY_SERIALIZABLE("messages", Messages, FMessage);
+		END_JSON_SERIALIZER
+	};
 
 	// Argument for AddMessageToConversation.
 	struct FAddMessageToConversationOptions : public FJsonSerializable
@@ -120,6 +317,11 @@ namespace UE::AIAssistant
 		// Message to add to the conversation.
 		FMessage Message;
 
+		bool operator==(const FAddMessageToConversationOptions& Other) const
+		{
+			return ConversationId == Other.ConversationId && Message == Other.Message;
+		}
+
 		BEGIN_JSON_SERIALIZER
 			JSON_SERIALIZE_OPTIONAL_OBJECT_SERIALIZABLE("conversationId", ConversationId);
 			JSON_SERIALIZE_OBJECT_SERIALIZABLE("message", Message);
@@ -402,4 +604,14 @@ namespace UE::AIAssistant
 		static const FString FunctionCallFormatTemplate;
 		static const FString FunctionCallbackFormatTemplate;
 	};
-}
\ No newline at end of file
+}
+
+FORCEINLINE uint32 GetTypeHash(const UE::AIAssistant::FConversationId& Id)
+{
+	return GetTypeHash(Id.Id);
+}
+
+FORCEINLINE uint32 GetTypeHash(const UE::AIAssistant::FMessageId& Id)
+{
+	return HashCombine(GetTypeHash(Id.Id), GetTypeHash(Id.Type));
+}
diff --git a/Engine/Plugins/Experimental/AIAssistant/Source/AIAssistant/Private/Tests/AIAssistantMessageUtils.cpp b/Engine/Plugins/Experimental/AIAssistant/Source/AIAssistant/Private/Tests/AIAssistantMessageUtils.cpp
new file mode 100644
--- /dev/null
+++ b/Engine/Plugins/Experimental/AIAssistant/Source/AIAssistant/Private/Tests/AIAssistantMessageUtils.cpp
@@ -0,0 +1,93 @@
+// Copyright Epic Games, Inc. All Rights Reserved.
+
+#include "Tests/AIAssistantMessageUtils.h"
+
+#include "AIAssistantWebApi.h"
+#include "Containers/UnrealString.h"
+#include "Misc/DateTime.h"
+#include "Misc/Optional.h"
+#include "Misc/Variant.h"
+#include "Templates/Tuple.h"
+
+#if WITH_DEV_AUTOMATION_TESTS
+
+namespace UE::AIAssistant
+{
+	FMessage MakeTextMessage(EMessageRole MessageRole, const FString& Text)
+	{
+		FMessage Message;
+		Message.MessageRole = MessageRole;
+		FMessageContent& MessageContent = Message.MessageContent.Emplace_GetRef();
+		MessageContent.ContentType = EMessageContentType::Text;
+		MessageContent.bVisibleToUser = true;
+		MessageContent.Content.Emplace<FTextMessageContent>();
+		MessageContent.Content.Get<FTextMessageContent>().Text = Text;
+		return Message;
+	}
+
+	FConversationId MakeConversationId(int Index)
+	{
+		FConversationId ConversationId;
+		ConversationId.Id = FString::Printf(TEXT("Conversation%i"), Index);
+		return ConversationId;
+	}
+
+	FMessageId MakeMessageId(int Index)
+	{
+		FMessageId MessageId;
+		MessageId.Id = FString::Printf(TEXT("Message%i"), Index);
+		MessageId.Type = TEXT("Message");
+		return MessageId;
+	}
+
+	FMessage MakeAgentTextMessage(int Index, const FString& Text)
+	{
+		FMessage Message = MakeTextMessage(EMessageRole::Agent, Text);
+		Message.MessageId = MakeMessageId(Index);
+		return Message;
+	}
+
+	FString MakeToolCallId(int MessageIndex, int ToolCallIndex)
+	{
+		return FString::Printf(TEXT("ToolCall%i-%i"), MessageIndex, ToolCallIndex);
+	}
+
+	FToolCallContent MakeToolCallContent(
+		FString Name,
+		FString JsonArguments,
+		TOptional<TTuple<int, int>> MessageAndToolCallIndex)
+	{
+		FToolCallContent ToolCallContent;
+		ToolCallContent.Name = Name;
+		ToolCallContent.JsonArguments = JsonArguments;
+		if (MessageAndToolCallIndex.IsSet())
+		{
+			const auto& [MessageIndex, ToolCallIndex] = MessageAndToolCallIndex.GetValue();
+			ToolCallContent.ToolCallId = MakeToolCallId(MessageIndex, ToolCallIndex);
+			ToolCallContent.ResponseRequired = true;
+		}
+		return ToolCallContent;
+	}
+
+	FMessage MakeAgentToolCallMessage(int Index, const TArray<FToolCallContent>& ToolCalls)
+	{
+		FMessage Message;
+		Message.MessageRole = EMessageRole::Agent;
+		Message.MessageId = MakeMessageId(Index);
+		for (const auto& ToolCallContent : ToolCalls)
+		{
+			FMessageContent& MessageContent = Message.MessageContent.Emplace_GetRef();
+			MessageContent.ContentType = EMessageContentType::ToolCall;
+			MessageContent.bVisibleToUser = false;
+			MessageContent.Content.Emplace<FToolCallContent>(ToolCallContent);
+		}
+		return Message;
+	}
+
+	FMessage MakeUserTextMessage(const FString& Text)
+	{
+		return MakeTextMessage(EMessageRole::User, Text);
+	}
+}
+
+#endif  // WITH_DEV_AUTOMATION_TESTS
diff --git a/Engine/Plugins/Experimental/AIAssistant/Source/AIAssistant/Private/Tests/AIAssistantMessageUtils.h b/Engine/Plugins/Experimental/AIAssistant/Source/AIAssistant/Private/Tests/AIAssistantMessageUtils.h
new file mode 100644
--- /dev/null
+++ b/Engine/Plugins/Experimental/AIAssistant/Source/AIAssistant/Private/Tests/AIAssistantMessageUtils.h
@@ -0,0 +1,29 @@
+// Copyright Epic Games, Inc. All Rights Reserved.
+
+#include "AIAssistantWebApi.h"
+
+#if WITH_DEV_AUTOMATION_TESTS
+
+namespace UE::AIAssistant
+{
+	FMessage MakeTextMessage(EMessageRole MessageRole, const FString& Text);
+	
+	FConversationId MakeConversationId(int Index);
+
+	FMessageId MakeMessageId(int Index);
+
+	FMessage MakeAgentTextMessage(int Index, const FString& Text);
+
+	FString MakeToolCallId(int MessageIndex, int ToolCallIndex);
+
+	FToolCallContent MakeToolCallContent(
+		FString Name,
+		FString JsonArguments,
+		TOptional<TTuple<int, int>> MessageAndToolCallIndex);
+
+	FMessage MakeAgentToolCallMessage(int Index, const TArray<FToolCallContent>& ToolCalls);
+
+	FMessage MakeUserTextMessage(const FString& Text);
+}
+
+#endif  // DEV_AUTOMATION_TESTS
diff --git a/Engine/Plugins/Experimental/AIAssistant/Source/AIAssistant/Private/Tests/AIAssistantWebApiTest.cpp b/Engine/Plugins/Experimental/AIAssistant/Source/AIAssistant/Private/Tests/AIAssistantWebApiTest.cpp
--- a/Engine/Plugins/Experimental/AIAssistant/Source/AIAssistant/Private/Tests/AIAssistantWebApiTest.cpp
+++ b/Engine/Plugins/Experimental/AIAssistant/Source/AIAssistant/Private/Tests/AIAssistantWebApiTest.cpp
@@ -9,6 +9,7 @@
 #include "AIAssistantFakeWebApi.h"
 #include "AIAssistantFakeWebJavaScriptExecutor.h"
 #include "AIAssistantFakeWebJavaScriptDelegateBinder.h"
+#include "AIAssistantMessageUtils.h"
 #include "AIAssistantTestFlags.h"
 #include "AIAssistantWebApi.h"
 #include "AIAssistantWebJavaScriptResultDelegateAccessor.h"
@@ -61,6 +62,849 @@ namespace UE::AIAssistant
 	};
 }
 
+BEGIN_DEFINE_SPEC(AIAssistantWebApiTextMessageContent, "AI.Assistant.WebApi.TextMessageContent",
+	AIAssistantTest::Flags)
+	FTextMessageContent Message;
+	FTextMessageContent Markdown;
+END_DEFINE_SPEC(AIAssistantWebApiTextMessageContent)
+
+void AIAssistantWebApiTextMessageContent::Define()
+{
+	BeforeEach([this]
+		{
+			Message.Text = TEXT("Hello");
+			Message.Format.Reset();
+
+			Markdown.Text = TEXT("Hello");
+			Markdown.Format = FString(TEXT("markdown"));
+		});
+
+	Describe("Equality operator", [this]
+		{
+			It("returns true when the object is identical", [this]
+				{
+					FTextMessageContent Same;
+					Same.Text = TEXT("Hello");
+					Same.Format.Reset();
+					TestTrue(TEXT("Message should be equal to itself"), Message == Message);
+					TestTrue(TEXT("Messages with same Text and unset Format should be equal"), Message == Same);
+				});
+			It("returns false when text is different", [this]
+				{
+					FTextMessageContent DifferentText;
+					DifferentText.Text = TEXT("Different");
+					DifferentText.Format.Reset();
+					TestFalse(TEXT("Messages with different Text should not be equal"), Message == DifferentText);
+				});
+			It("returns true when formats are present and identical", [this]
+				{
+					FTextMessageContent MarkdownSame;
+					MarkdownSame.Text = TEXT("Hello");
+					MarkdownSame.Format = FString(TEXT("markdown"));
+					TestTrue(TEXT("Messages with same Text and same Format should be equal"),
+						Markdown == MarkdownSame);
+				});
+			It("returns false when formats are present and different", [this]
+				{
+					FTextMessageContent Plain;
+					Plain.Text = TEXT("Hello");
+					Plain.Format = FString(TEXT("plain"));
+					TestFalse(TEXT("Messages with same Text but different Format should not be equal"),
+						Markdown == Plain);
+				});
+			It("treats unset formats as equal to 'markdown'", [this]
+				{
+					FTextMessageContent WithoutFormat;
+					WithoutFormat.Text = TEXT("Hello");
+					WithoutFormat.Format.Reset();
+
+					TestTrue(TEXT("Messages with absent Format should be treated as markdown"),
+						Markdown == WithoutFormat);
+					TestTrue(TEXT("Messages with absent Format should be treated as markdown"),
+						WithoutFormat == Markdown);
+				});
+		});
+}
+
+BEGIN_DEFINE_SPEC(AIAssistantWebApiToolCallContent, "AI.Assistant.WebApi.ToolCallContent",
+	AIAssistantTest::Flags)
+	FToolCallContent ToolCall;
+	FToolCallContent WithIdAndResponse;
+END_DEFINE_SPEC(AIAssistantWebApiToolCallContent)
+
+void AIAssistantWebApiToolCallContent::Define()
+{
+	BeforeEach([this]
+		{
+			ToolCall.Name = TEXT("agent_status");
+			ToolCall.JsonArguments = TEXT(R"json({"foo": 1})json");
+			ToolCall.ToolCallId.Reset();
+			ToolCall.ResponseRequired.Reset();
+
+			WithIdAndResponse.Name = TEXT("agent_status");
+			WithIdAndResponse.JsonArguments = TEXT(R"json({"foo": 1})json");
+			WithIdAndResponse.ToolCallId = TEXT("ToolCall1-0");
+			WithIdAndResponse.ResponseRequired = true;
+		});
+
+	Describe("Equality operator", [this]
+		{
+			It("returns true when the object is identical", [this]
+				{
+					FToolCallContent ToolCallSame;
+					ToolCallSame.Name = TEXT("agent_status");
+					ToolCallSame.JsonArguments = TEXT(R"json({"foo": 1})json");
+					ToolCallSame.ToolCallId.Reset();
+					ToolCallSame.ResponseRequired.Reset();
+					TestTrue(TEXT("Tool call should be equal to itself"), ToolCall == ToolCall);
+					TestTrue(TEXT("Tool calls with same fields and unset optionals should be equal"),
+						ToolCall == ToolCallSame);
+				});
+			It("returns false when the name is different", [this]
+				{
+					FToolCallContent DifferentName;
+					DifferentName.Name = TEXT("other_tool");
+					DifferentName.JsonArguments = TEXT(R"json({"foo": 1})json");
+					DifferentName.ToolCallId.Reset();
+					DifferentName.ResponseRequired.Reset();
+					TestFalse(TEXT("Tool calls with different Name should not be equal"),
+						ToolCall == DifferentName);
+				});
+			It("returns false when the JsonArguments are different", [this]
+				{
+					FToolCallContent DifferentArguments;
+					DifferentArguments.Name = TEXT("agent_status");
+					DifferentArguments.JsonArguments = TEXT(R"json({"foo": 2})json");
+					DifferentArguments.ToolCallId.Reset();
+					DifferentArguments.ResponseRequired.Reset();
+					TestFalse(TEXT("Tool calls with different JsonArguments should not be equal"),
+						ToolCall == DifferentArguments);
+				});
+			It("returns true when the ToolCallId and ResponseRequired are set", [this]
+				{
+					FToolCallContent WithIdAndResponseSame;
+					WithIdAndResponseSame.Name = TEXT("agent_status");
+					WithIdAndResponseSame.JsonArguments = TEXT(R"json({"foo": 1})json");
+					WithIdAndResponseSame.ToolCallId = TEXT("ToolCall1-0");
+					WithIdAndResponseSame.ResponseRequired = true;
+					TestTrue(TEXT("Tool calls with same ToolCallId and ResponseRequired should be equal"),
+						WithIdAndResponse == WithIdAndResponseSame);
+				});
+			It("returns false when the ToolCallIds are different", [this]
+				{
+					FToolCallContent DifferentToolCallId;
+					DifferentToolCallId.Name = TEXT("agent_status");
+					DifferentToolCallId.JsonArguments = TEXT(R"json({"foo": 1})json");
+					DifferentToolCallId.ToolCallId = TEXT("ToolCall1-1");
+					DifferentToolCallId.ResponseRequired = true;
+					TestFalse(TEXT("Tool calls with different ToolCallId should not be equal"),
+						WithIdAndResponse == DifferentToolCallId);
+				});
+			It("returns false when the ResponseRequireds are different", [this]
+				{
+					FToolCallContent DifferentResponseRequired;
+					DifferentResponseRequired.Name = TEXT("agent_status");
+					DifferentResponseRequired.JsonArguments = TEXT(R"json({"foo": 1})json");
+					DifferentResponseRequired.ToolCallId = TEXT("ToolCall1-0");
+					DifferentResponseRequired.ResponseRequired = false;
+					TestFalse(TEXT("Tool calls with different ResponseRequired should not be equal"),
+						WithIdAndResponse == DifferentResponseRequired);
+				});
+			It("returns false when one ResponseRequired is set and the other is absent", [this]
+				{
+					FToolCallContent WithResponseRequired;
+					WithResponseRequired.Name = TEXT("agent_status");
+					WithResponseRequired.JsonArguments = TEXT(R"json({"foo": 1})json");
+					WithResponseRequired.ToolCallId.Reset();
+					WithResponseRequired.ResponseRequired = true;
+
+					FToolCallContent WithoutResponseRequired = WithResponseRequired;
+					WithoutResponseRequired.ResponseRequired.Reset();
+
+					TestFalse(
+						TEXT("Presence/absence of ResponseRequired should affect equality"),
+						WithResponseRequired == WithoutResponseRequired);
+				});
+			It("returns false when one ToolCallId is set and the other is absent", [this]
+				{
+					FToolCallContent WithIdOnly;
+					WithIdOnly.Name = TEXT("agent_status");
+					WithIdOnly.JsonArguments = TEXT(R"json({"foo": 1})json");
+					WithIdOnly.ToolCallId = TEXT("ToolCall1-0");
+					WithIdOnly.ResponseRequired.Reset();
+					TestFalse(TEXT("Presence/absence of ToolCallId should affect equality"),
+						WithIdOnly == ToolCall);
+				});
+		});
+}
+
+
+BEGIN_DEFINE_SPEC(AIAssistantWebApiToolResponseContent, "AI.Assistant.WebApi.ToolResponseContent",
+	AIAssistantTest::Flags)
+	FToolResponseContent Response;
+END_DEFINE_SPEC(AIAssistantWebApiToolResponseContent)
+
+void AIAssistantWebApiToolResponseContent::Define()
+{
+	BeforeEach([this]
+		{
+			Response.ToolCallId = TEXT("ToolCall1-0");
+			Response.Name = TEXT("agent_status");
+			Response.JsonResponse = TEXT(R"json({"ok": true})json");
+			Response.Success.Reset();
+			Response.ErrorMessage.Reset();
+		});
+
+	Describe("Equality operator", [this]
+		{
+			It("returns true when the object is identical", [this]
+				{
+					FToolResponseContent ResponseSame;
+					ResponseSame.ToolCallId = TEXT("ToolCall1-0");
+					ResponseSame.Name = TEXT("agent_status");
+					ResponseSame.JsonResponse = TEXT(R"json({"ok": true})json");
+					ResponseSame.Success.Reset();
+					ResponseSame.ErrorMessage.Reset();
+
+					TestTrue(TEXT("Tool response should be equal to itself"), Response == Response);
+					TestTrue(TEXT("Tool responses with same basic fields and unset optionals should be equal"),
+						Response == ResponseSame);
+				});
+			It("returns false when the ToolCallId differs", [this]
+				{
+					FToolResponseContent DifferentToolCallId;
+					DifferentToolCallId.ToolCallId = TEXT("ToolCall1-1");
+					DifferentToolCallId.Name = TEXT("agent_status");
+					DifferentToolCallId.JsonResponse = TEXT(R"json({"ok": true})json");
+					DifferentToolCallId.Success.Reset();
+					DifferentToolCallId.ErrorMessage.Reset();
+					TestFalse(TEXT("Tool responses with different ToolCallId should not be equal"),
+						Response == DifferentToolCallId);
+				});
+			It("returns false when the name differs", [this]
+				{
+					FToolResponseContent DifferentName;
+					DifferentName.ToolCallId = TEXT("ToolCall1-0");
+					DifferentName.Name = TEXT("other_tool");
+					DifferentName.JsonResponse = TEXT(R"json({"ok": true})json");
+					DifferentName.Success.Reset();
+					DifferentName.ErrorMessage.Reset();
+					TestFalse(TEXT("Tool responses with different Name should not be equal"),
+						Response == DifferentName);
+				});
+			It("returns false when the JsonResponse differs", [this]
+				{
+					FToolResponseContent DifferentJson;
+					DifferentJson.ToolCallId = TEXT("ToolCall1-0");
+					DifferentJson.Name = TEXT("agent_status");
+					DifferentJson.JsonResponse = TEXT(R"json({"ok": false})json");
+					DifferentJson.Success.Reset();
+					DifferentJson.ErrorMessage.Reset();
+					TestFalse(TEXT("Tool responses with different JsonResponse should not be equal"),
+						Response == DifferentJson);
+				});
+			It("returns true because it treats unset Success value as true", [this]
+				{
+					FToolResponseContent SuccessTrue;
+					SuccessTrue.ToolCallId = TEXT("ToolCall1-0");
+					SuccessTrue.Name = TEXT("agent_status");
+					SuccessTrue.JsonResponse = TEXT(R"json({"ok": true})json");
+					SuccessTrue.Success = true;
+					SuccessTrue.ErrorMessage.Reset();
+					TestTrue(TEXT("Tool responses with unset Success and a Success set to true should be equal"),
+						Response == SuccessTrue);
+
+					FToolResponseContent SuccessTrueSame = SuccessTrue;
+					TestTrue(TEXT("Tool responses with same Success and ErrorMessage should be equal"),
+						SuccessTrue == SuccessTrueSame);
+				});
+			It("returns false when Success does not match", [this]
+				{
+					FToolResponseContent SuccessTrue;
+					SuccessTrue.ToolCallId = TEXT("ToolCall1-0");
+					SuccessTrue.Name = TEXT("agent_status");
+					SuccessTrue.JsonResponse = TEXT(R"json({"ok": true})json");
+					SuccessTrue.Success = true;
+					SuccessTrue.ErrorMessage.Reset();
+
+					FToolResponseContent SuccessFalse;
+					SuccessFalse.ToolCallId = TEXT("ToolCall1-0");
+					SuccessFalse.Name = TEXT("agent_status");
+					SuccessFalse.JsonResponse = TEXT(R"json({"ok": true})json");
+					SuccessFalse.Success = false;
+					SuccessFalse.ErrorMessage.Reset();
+					TestFalse(TEXT("Tool responses with different Success should not be equal"),
+						Response == SuccessFalse);
+					TestFalse(TEXT("Tool responses with different Success should not be equal"),
+						SuccessTrue == SuccessFalse);
+				});
+			It("returns true with identical ErrorMessages", [this]
+				{
+					FToolResponseContent WithError;
+					WithError.ToolCallId = TEXT("ToolCall1-0");
+					WithError.Name = TEXT("agent_status");
+					WithError.JsonResponse = TEXT(R"json({"ok": false})json");
+					WithError.Success = false;
+					WithError.ErrorMessage = FString(TEXT("error"));
+					
+					FToolResponseContent WithErrorSame = WithError;
+
+					TestTrue(TEXT("Tool responses with same ErrorMessage should be equal"),
+						WithError == WithErrorSame);
+				});
+			It("returns false when error messages differ", [this]
+				{
+					FToolResponseContent WithError;
+					WithError.ToolCallId = TEXT("ToolCall1-0");
+					WithError.Name = TEXT("agent_status");
+					WithError.JsonResponse = TEXT(R"json({"ok": false})json");
+					WithError.Success = false;
+					WithError.ErrorMessage = FString(TEXT("error"));
+
+					FToolResponseContent DifferentError = WithError;
+					DifferentError.ErrorMessage = FString(TEXT("other error"));
+					
+					TestFalse(TEXT("Tool responses with different ErrorMessage should not be equal"),
+						WithError == DifferentError);
+				});
+		});
+}
+
+
+BEGIN_DEFINE_SPEC(AIAssistantWebApiMessageContent, "AI.Assistant.WebApi.MessageContent",
+	AIAssistantTest::Flags)
+	FTextMessageContent HelloText;
+	FMessageContent TextContent;
+END_DEFINE_SPEC(AIAssistantWebApiMessageContent)
+
+void AIAssistantWebApiMessageContent::Define()
+{
+	BeforeEach([this]
+		{
+			HelloText.Text = TEXT("Hello");
+			HelloText.Format.Reset();
+
+			TextContent.ContentType = EMessageContentType::Text;
+			TextContent.bVisibleToUser = true;
+			TextContent.Content.Emplace<FTextMessageContent>(HelloText);
+		});
+
+	Describe("Equality operator", [this]
+		{
+			It("returns true when the object is identical", [this]
+				{
+					FMessageContent TextContentSame;
+					TextContentSame.ContentType = EMessageContentType::Text;
+					TextContentSame.bVisibleToUser = true;
+					TextContentSame.Content.Emplace<FTextMessageContent>(HelloText);
+
+					TestTrue(TEXT("Message content should be equal to itself"), TextContent == TextContent);
+					TestTrue(TEXT("Message contents with same Text variant and visibility should be equal"),
+						TextContent == TextContentSame);
+				});
+			It("returns false when text is different", [this]
+				{
+					FTextMessageContent DifferentText;
+					DifferentText.Text = TEXT("Different");
+					DifferentText.Format.Reset();
+
+					FMessageContent DifferentTextContent;
+					DifferentTextContent.ContentType = EMessageContentType::Text;
+					DifferentTextContent.bVisibleToUser = true;
+					DifferentTextContent.Content.Emplace<FTextMessageContent>(DifferentText);
+
+					TestFalse(TEXT("Message contents with different Text should not be equal"),
+						TextContent == DifferentTextContent);
+				});
+			It("returns false when visibility is different", [this]
+				{
+					FMessageContent DifferentVisibility;
+					DifferentVisibility.ContentType = EMessageContentType::Text;
+					DifferentVisibility.bVisibleToUser = false;
+					DifferentVisibility.Content.Emplace<FTextMessageContent>(HelloText);
+
+					TestFalse(TEXT("Message contents with different visibility should not be equal"),
+						TextContent == DifferentVisibility);
+				});
+			It("returns false when the content type is different", [this]
+				{
+					FToolCallContent ToolCall;
+					ToolCall.Name = TEXT("tool");
+					ToolCall.JsonArguments = TEXT(R"json({})json");
+					ToolCall.ToolCallId.Reset();
+					ToolCall.ResponseRequired.Reset();
+
+					FMessageContent ToolCallContent;
+					ToolCallContent.ContentType = EMessageContentType::ToolCall;
+					ToolCallContent.bVisibleToUser = false;
+					ToolCallContent.Content.Emplace<FToolCallContent>(ToolCall);
+
+					TestFalse(TEXT("Message contents with different ContentType should not be equal"),
+						TextContent == ToolCallContent);
+				});
+		});
+}
+
+
+BEGIN_DEFINE_SPEC(AIAssistantWebApiMessageId, "AI.Assistant.WebApi.MessageId",
+	AIAssistantTest::Flags)
+	FMessageId MessageId;
+END_DEFINE_SPEC(AIAssistantWebApiMessageId)
+
+void AIAssistantWebApiMessageId::Define()
+{
+	BeforeEach([this]
+		{
+			MessageId = MakeMessageId(1);
+		});
+
+	Describe("Equality operator", [this]
+		{
+			It("returns true when the object is identical", [this]
+				{
+					FMessageId MessageIdSame = MakeMessageId(1);
+					TestTrue(TEXT("MessageId should be equal to itself"), MessageId == MessageId);
+					TestTrue(TEXT("MessageIds created with same index should be equal"),
+						MessageId == MessageIdSame);
+				});
+			It("returns false when the Id is different", [this]
+				{
+					FMessageId MessageIdDifferentIndex = MakeMessageId(2);
+					TestFalse(TEXT("MessageIds with different Id should not be equal"),
+						MessageId == MessageIdDifferentIndex);
+				});
+			It("returns false when the Type is different", [this]
+				{
+					FMessageId MessageIdDifferentType = MakeMessageId(1);
+					MessageIdDifferentType.Type = TEXT("OtherType");
+					TestFalse(TEXT("MessageIds with different Type should not be equal"),
+						MessageId == MessageIdDifferentType);
+				});
+		});
+
+	Describe("hash function", [this]
+		{
+			It("produces same hash for same MessageId", [this]
+				{
+					FMessageId MessageIdSame = MakeMessageId(1);
+
+					const uint32 Hash = GetTypeHash(MessageId);
+					const uint32 HashSame = GetTypeHash(MessageIdSame);
+
+					TestEqual(TEXT("Same MessageId value should produce same hash"), Hash, HashSame);
+				});
+			It("produces different hash for different MessageId", [this]
+				{
+					FMessageId MessageIdDifferent = MakeMessageId(2);
+
+					const uint32 Hash = GetTypeHash(MessageId);
+					const uint32 HashDifferent = GetTypeHash(MessageIdDifferent);
+
+					TestNotEqual(TEXT("Different MessageId value should produce different hash"),
+						Hash, HashDifferent);
+				});
+		});
+}
+
+BEGIN_DEFINE_SPEC(AIAssistantWebApiMessage, "AI.Assistant.WebApi.Message", AIAssistantTest::Flags)
+	FMessage Message;
+END_DEFINE_SPEC(AIAssistantWebApiMessage)
+
+void AIAssistantWebApiMessage::Define()
+{
+	BeforeEach([this]
+		{
+			Message = MakeAgentTextMessage(1, TEXT("Hello"));
+		});
+
+	Describe("Equality operator", [this]
+		{
+			It("returns true when the object is identical", [this]
+				{
+					FMessage MessageSame = MakeAgentTextMessage(1, TEXT("Hello"));
+					TestTrue(TEXT("Message should be equal to itself"), Message == Message);
+					TestTrue(TEXT("Messages created with same helper parameters should be equal"),
+						Message == MessageSame);
+				});
+
+			It("returns false when the MessageId differs", [this]
+				{
+					FMessage DifferentMessageId = MakeAgentTextMessage(2, TEXT("Hello"));
+					TestFalse(TEXT("Messages with different MessageId should not be equal"),
+						Message == DifferentMessageId);
+				});
+
+			It("returns false when a present MessageId is compared to an absent MessageId", [this]
+				{
+					FMessage NoMessageId = MakeAgentTextMessage(1, TEXT("Hello"));
+					NoMessageId.MessageId.Reset();
+
+					TestFalse(TEXT("Messages differing by presence/absence of MessageId should not be equal"),
+						Message == NoMessageId);
+				});
+
+			It("returns false when the roles differ", [this]
+				{
+					FMessage DifferentRole = MakeTextMessage(EMessageRole::User, TEXT("Hello"));
+					DifferentRole.MessageId = MakeMessageId(1);
+					TestFalse(TEXT("Messages with different MessageRole should not be equal"),
+						Message == DifferentRole);
+				});
+
+			It("returns false when the text content differs", [this]
+				{
+					FMessage DifferentText = MakeAgentTextMessage(1, TEXT("Other text"));
+					TestFalse(TEXT("Messages with different text content should not be equal"),
+						Message == DifferentText);
+				});
+
+			It("returns false when the content count differs", [this]
+				{
+					FMessage ExtraContent = MakeAgentTextMessage(1, TEXT("Hello"));
+					// Add an extra content entry so the counts differ
+					ExtraContent.MessageContent.Add(Message.MessageContent[0]);
+
+					TestFalse(TEXT("Messages with different number of MessageContents should not be equal"),
+						Message == ExtraContent);
+				});
+
+			It("returns true when the date matches", [this]
+				{
+					FMessage WithDate = MakeAgentTextMessage(1, TEXT("Hello"));
+					WithDate.Date = FDateTime(2024, 1, 1);
+					FMessage WithDateSame = MakeAgentTextMessage(1, TEXT("Hello"));
+					WithDateSame.Date = FDateTime(2024, 1, 1);
+					TestTrue(TEXT("Messages with same Date should be equal"), WithDate == WithDateSame);
+				});
+
+			It("returns false when the date differs", [this]
+				{
+					FMessage WithDate = MakeAgentTextMessage(1, TEXT("Hello"));
+					WithDate.Date = FDateTime(2024, 1, 1);
+					FMessage DifferentDate = MakeAgentTextMessage(1, TEXT("Hello"));
+					DifferentDate.Date = FDateTime(2024, 1, 2);
+					TestFalse(TEXT("Messages with different Date should not be equal"),
+						WithDate == DifferentDate);
+				});
+
+			It("returns false when one date is set and the other is absent", [this]
+				{
+					FMessage WithDate = MakeAgentTextMessage(1, TEXT("Hello"));
+					WithDate.Date = FDateTime(2024, 1, 1);
+					FMessage DateUnset = MakeAgentTextMessage(1, TEXT("Hello"));
+					DateUnset.Date.Reset();
+					TestFalse(TEXT("Messages differing by presence/absence of Date should not be equal"),
+						WithDate == DateUnset);
+				});
+		});
+}
+
+BEGIN_DEFINE_SPEC(AIAssistantWebApiConversationId, "AI.Assistant.WebApi.ConversationId",
+	AIAssistantTest::Flags)
+	FConversationId ConversationId;
+END_DEFINE_SPEC(AIAssistantWebApiConversationId)
+
+void AIAssistantWebApiConversationId::Define()
+{
+	BeforeEach([this]
+		{
+			ConversationId = MakeConversationId(1);
+		});
+
+	Describe("Equality operator", [this]
+		{
+			It("returns true when the object is identical", [this]
+				{
+					FConversationId ConversationIdSame = MakeConversationId(1);
+					TestTrue(TEXT("ConversationId should be equal to itself"),
+						ConversationId == ConversationId);
+					TestTrue(TEXT("ConversationIds with same Id should be equal"),
+						ConversationId == ConversationIdSame);
+				});
+
+			It("returns false when the Id is different", [this]
+				{
+					FConversationId ConversationIdDifferent = MakeConversationId(2);
+					TestFalse(TEXT("ConversationIds with different Id should not be equal"),
+						ConversationId == ConversationIdDifferent);
+				});
+		});
+
+	Describe("hash function", [this]
+		{
+			It("produces same hash for same ConversationId", [this]
+				{
+					FConversationId ConversationIdSame = MakeConversationId(1);
+
+					const uint32 Hash = GetTypeHash(ConversationId);
+					const uint32 HashSame = GetTypeHash(ConversationIdSame);
+
+					TestEqual(TEXT("Same ConversationId value should produce same hash"), Hash, HashSame);
+				});
+
+			It("produces different hash for different ConversationId", [this]
+				{
+					FConversationId ConversationIdDifferent = MakeConversationId(2);
+
+					const uint32 Hash = GetTypeHash(ConversationId);
+					const uint32 HashDifferent = GetTypeHash(ConversationIdDifferent);
+
+					TestNotEqual(TEXT("Different ConversationId value should produce different hash"),
+						Hash, HashDifferent);
+				});
+		});
+}
+
+BEGIN_DEFINE_SPEC(AIAssistantWebApiConversationDescriptor,
+	"AI.Assistant.WebApi.ConversationDescriptor", AIAssistantTest::Flags)
+	FConversationDescriptor Descriptor;
+	FConversationDescriptor DescriptorWithConversationId;
+END_DEFINE_SPEC(AIAssistantWebApiConversationDescriptor)
+
+void AIAssistantWebApiConversationDescriptor::Define()
+{
+	BeforeEach([this]
+		{
+			// Base descriptor without ConversationId
+			Descriptor.ConversationId.Reset();
+			Descriptor.DisplayName = TEXT("My Conversation");
+			Descriptor.Labels = { TEXT("red"), TEXT("blue") };
+			Descriptor.CreatedAt = FDateTime(2000, 1, 1);
+			Descriptor.UpdatedAt = FDateTime(2025, 1, 1);
+
+			// Descriptor with ConversationId
+			DescriptorWithConversationId = Descriptor;
+			DescriptorWithConversationId.ConversationId = MakeConversationId(1);
+		});
+
+	Describe("Equality operator", [this]
+		{
+			It("returns true when descriptors have same values and no ConversationId", [this]
+				{
+					FConversationDescriptor DescriptorSame = Descriptor;
+
+					TestTrue(TEXT("ConversationDescriptor should be equal to itself"),
+						Descriptor == Descriptor);
+					TestTrue(TEXT("Descriptors with same values and no ConversationId should be equal"),
+						Descriptor == DescriptorSame);
+				});
+
+			It("returns true when descriptors have same ConversationId", [this]
+				{
+					FConversationDescriptor WithConversationIdSame = DescriptorWithConversationId;
+
+					TestTrue(TEXT("Descriptors with same ConversationId should be equal"),
+						DescriptorWithConversationId == WithConversationIdSame);
+				});
+
+			It("returns false when the ConversationId differs", [this]
+				{
+					FConversationDescriptor DifferentConversationId = DescriptorWithConversationId;
+					DifferentConversationId.ConversationId->Id = TEXT("Conversation2");
+
+					TestFalse(TEXT("Descriptors with different ConversationId should not be equal"),
+						DescriptorWithConversationId == DifferentConversationId);
+				});
+
+			It("returns false when one ConversationId is set and the other is absent", [this]
+				{
+					FConversationDescriptor ConversationIdMissing = DescriptorWithConversationId;
+					ConversationIdMissing.ConversationId.Reset();
+
+					TestFalse(TEXT("Presence/absence of ConversationId should affect equality"),
+						DescriptorWithConversationId == ConversationIdMissing);
+				});
+
+			It("returns false when the DisplayName differs", [this]
+				{
+					FConversationDescriptor DifferentDisplayName = Descriptor;
+					DifferentDisplayName.DisplayName = TEXT("Other Conversation");
+					TestFalse(TEXT("Descriptors with different DisplayName should not be equal"),
+						Descriptor == DifferentDisplayName);
+				});
+
+			It("returns false when Labels content differs", [this]
+				{
+					FConversationDescriptor DifferentLabels = Descriptor;
+					DifferentLabels.Labels[0] = TEXT("green");
+					TestFalse(TEXT("Descriptors with different Labels content should not be equal"),
+						Descriptor == DifferentLabels);
+				});
+
+			It("returns false when Labels size differs", [this]
+				{
+					FConversationDescriptor ExtraLabel = Descriptor;
+					ExtraLabel.Labels.Add(TEXT("yellow"));
+					TestFalse(TEXT("Descriptors with different Labels size should not be equal"),
+						Descriptor == ExtraLabel);
+				});
+
+			It("returns false when CreatedAt differs", [this]
+				{
+					FConversationDescriptor DifferentCreatedAt = Descriptor;
+					DifferentCreatedAt.CreatedAt = FDateTime(2000, 1, 2);
+					TestFalse(TEXT("Descriptors with different CreatedAt should not be equal"),
+						Descriptor == DifferentCreatedAt);
+				});
+
+			It("returns false when UpdatedAt differs", [this]
+				{
+					FConversationDescriptor DifferentUpdatedAt = Descriptor;
+					DifferentUpdatedAt.UpdatedAt = FDateTime(2025, 1, 2);
+					TestFalse(TEXT("Descriptors with different UpdatedAt should not be equal"),
+						Descriptor == DifferentUpdatedAt);
+				});
+		});
+}
+
+BEGIN_DEFINE_SPEC(AIAssistantWebApiConversation, "AI.Assistant.WebApi.Conversation",
+	AIAssistantTest::Flags)
+	FConversationDescriptor Descriptor;
+	FMessage AgentMessage;
+	FMessage UserMessage;
+	FConversation Conversation;
+END_DEFINE_SPEC(AIAssistantWebApiConversation)
+
+void AIAssistantWebApiConversation::Define()
+{
+	BeforeEach([this]
+		{
+			// Base descriptor
+			Descriptor.ConversationId = MakeConversationId(1);
+			Descriptor.DisplayName = TEXT("My Conversation");
+			Descriptor.Labels = { TEXT("red") };
+			Descriptor.CreatedAt = FDateTime(2000, 1, 1);
+			Descriptor.UpdatedAt = FDateTime(2025, 1, 1);
+
+			// Base messages
+			AgentMessage = MakeAgentTextMessage(1, TEXT("Hello"));
+			UserMessage = MakeUserTextMessage(TEXT("Hi"));
+
+			// Base conversation
+			Conversation.Descriptor = Descriptor;
+			Conversation.Messages = { AgentMessage, UserMessage };
+		});
+
+	Describe("Equality operator", [this]
+		{
+			It("returns true when the object is identical", [this]
+				{
+					FConversationDescriptor DescriptorSame = Descriptor;
+
+					FConversation ConversationSame;
+					ConversationSame.Descriptor = DescriptorSame;
+					ConversationSame.Messages = {
+						MakeAgentTextMessage(1, TEXT("Hello")),
+						MakeUserTextMessage(TEXT("Hi")),
+					};
+
+					TestTrue(TEXT("Conversation should be equal to itself"),
+						Conversation == Conversation);
+					TestTrue(TEXT("Conversations with same descriptor and messages should be equal"),
+						Conversation == ConversationSame);
+				});
+
+			It("returns false when the descriptor differs", [this]
+				{
+					FConversationDescriptor DifferentDescriptor = Descriptor;
+					DifferentDescriptor.DisplayName = TEXT("Other Conversation");
+
+					FConversation ConversationDifferentDescriptor;
+					ConversationDifferentDescriptor.Descriptor = DifferentDescriptor;
+					ConversationDifferentDescriptor.Messages = { AgentMessage, UserMessage };
+
+					TestFalse(TEXT("Conversations with different descriptor should not be equal"),
+						Conversation == ConversationDifferentDescriptor);
+				});
+
+			It("returns false when the message contents differ", [this]
+				{
+					FConversation ConversationDifferentMessages;
+					ConversationDifferentMessages.Descriptor = Descriptor;
+					ConversationDifferentMessages.Messages = {
+						MakeAgentTextMessage(1, TEXT("Other text")),
+						UserMessage
+					};
+
+					TestFalse(TEXT("Conversations with different message content should not be equal"),
+						Conversation == ConversationDifferentMessages);
+				});
+
+			It("returns false when the number of messages differs", [this]
+				{
+					FConversation ConversationExtraMessage;
+					ConversationExtraMessage.Descriptor = Descriptor;
+					ConversationExtraMessage.Messages = {
+						AgentMessage,
+						UserMessage,
+						MakeUserTextMessage(TEXT("Extra"))
+					};
+
+					TestFalse(TEXT("Conversations with different number of messages should not be equal"),
+						Conversation == ConversationExtraMessage);
+				});
+		});
+}
+
+BEGIN_DEFINE_SPEC(AIAssistantWebApiAddMessageToConversationOptions,
+	"AI.Assistant.WebApi.AddMessageToConversationOptions", AIAssistantTest::Flags)
+	FAddMessageToConversationOptions Options;
+END_DEFINE_SPEC(AIAssistantWebApiAddMessageToConversationOptions)
+
+void AIAssistantWebApiAddMessageToConversationOptions::Define()
+{
+	BeforeEach([this]
+		{
+			// Base options
+			Options.ConversationId = MakeConversationId(1);
+			Options.Message = MakeAgentTextMessage(1, TEXT("Hello"));
+		});
+
+	Describe("Equality operator", [this]
+		{
+			It("returns true when the object is identical", [this]
+				{
+					FAddMessageToConversationOptions OptionsSame;
+					OptionsSame.ConversationId = MakeConversationId(1);
+					OptionsSame.Message = MakeAgentTextMessage(1, TEXT("Hello"));
+
+					TestTrue(TEXT("Options should be equal to itself"), Options == Options);
+					TestTrue(TEXT("Options with same ConversationId and Message should be equal"),
+						Options == OptionsSame);
+				});
+
+			It("returns false when the ConversationId differs", [this]
+				{
+					FAddMessageToConversationOptions DifferentConversationId;
+					DifferentConversationId.ConversationId = MakeConversationId(2);
+					DifferentConversationId.Message = MakeAgentTextMessage(1, TEXT("Hello"));
+
+					TestFalse(TEXT("Options with different ConversationId should not be equal"),
+						Options == DifferentConversationId);
+				});
+
+			It("returns false when one ConversationId is set and the other is absent", [this]
+				{
+					FAddMessageToConversationOptions ConversationIdMissing;
+					ConversationIdMissing.ConversationId.Reset();
+					ConversationIdMissing.Message = MakeAgentTextMessage(1, TEXT("Hello"));
+
+					TestFalse(TEXT("Options differing by presence/absence of ConversationId should not be equal"),
+						Options == ConversationIdMissing);
+				});
+
+			It("returns false when the Message differs", [this]
+				{
+					FAddMessageToConversationOptions DifferentMessage;
+					DifferentMessage.ConversationId = MakeConversationId(1);
+					DifferentMessage.Message = MakeAgentTextMessage(1, TEXT("Other text"));
+
+					TestFalse(TEXT("Options with different Message should not be equal"),
+						Options == DifferentMessage);
+				});
+		});
+}
+
 IMPLEMENT_SIMPLE_AUTOMATION_TEST(
 	FAIAssistantWebApiTestFormatFunctionCallGlobalNoArgs,
 	"AI.Assistant.WebApi.FormatFunctionCallGlobalNoArgs",
