Fix PVS Studio warnings in Unit tests

PVS Studio warns when there are identical expressions on both sides of operator==. We want to do this intentionally for unit tests though so that we can test identity, so we squelch them in the places where we run identity tests.

#Synced-CL 48864563

diff --git a/Engine/Plugins/Experimental/AIAssistant/Source/AIAssistant/Private/Tests/AIAssistantWebApiTest.cpp b/Engine/Plugins/Experimental/AIAssistant/Source/AIAssistant/Private/Tests/AIAssistantWebApiTest.cpp
--- a/Engine/Plugins/Experimental/AIAssistant/Source/AIAssistant/Private/Tests/AIAssistantWebApiTest.cpp
+++ b/Engine/Plugins/Experimental/AIAssistant/Source/AIAssistant/Private/Tests/AIAssistantWebApiTest.cpp
@@ -86,6 +86,8 @@ void AIAssistantWebApiTextMessageContent::Define()
 					FTextMessageContent Same;
 					Same.Text = TEXT("Hello");
 					Same.Format.Reset();
+					// Intentional self-compare in test.
+					//-V501
 					TestTrue(TEXT("Message should be equal to itself"), Message == Message);
 					TestTrue(TEXT("Messages with same Text and unset Format should be equal"), Message == Same);
 				});
@@ -156,6 +158,8 @@ void AIAssistantWebApiToolCallContent::Define()
 					ToolCallSame.JsonArguments = TEXT(R"json({"foo": 1})json");
 					ToolCallSame.ToolCallId.Reset();
 					ToolCallSame.ResponseRequired.Reset();
+					// Intentional self-compare in test.
+					//-V501
 					TestTrue(TEXT("Tool call should be equal to itself"), ToolCall == ToolCall);
 					TestTrue(TEXT("Tool calls with same fields and unset optionals should be equal"),
 						ToolCall == ToolCallSame);
@@ -265,7 +269,9 @@ void AIAssistantWebApiToolResponseContent::Define()
 					ResponseSame.JsonResponse = TEXT(R"json({"ok": true})json");
 					ResponseSame.Success.Reset();
 					ResponseSame.ErrorMessage.Reset();
-
+					
+					// Intentional self-compare in test.
+					//-V501
 					TestTrue(TEXT("Tool response should be equal to itself"), Response == Response);
 					TestTrue(TEXT("Tool responses with same basic fields and unset optionals should be equal"),
 						Response == ResponseSame);
@@ -397,7 +403,9 @@ void AIAssistantWebApiMessageContent::Define()
 					TextContentSame.ContentType = EMessageContentType::Text;
 					TextContentSame.bVisibleToUser = true;
 					TextContentSame.Content.Emplace<FTextMessageContent>(HelloText);
-
+					
+					// Intentional self-compare in test.
+					//-V501
 					TestTrue(TEXT("Message content should be equal to itself"), TextContent == TextContent);
 					TestTrue(TEXT("Message contents with same Text variant and visibility should be equal"),
 						TextContent == TextContentSame);
@@ -463,6 +471,8 @@ void AIAssistantWebApiMessageId::Define()
 			It("returns true when the object is identical", [this]
 				{
 					FMessageId MessageIdSame = MakeMessageId(1);
+					// Intentional self-compare in test.
+					//-V501
 					TestTrue(TEXT("MessageId should be equal to itself"), MessageId == MessageId);
 					TestTrue(TEXT("MessageIds created with same index should be equal"),
 						MessageId == MessageIdSame);
@@ -522,6 +532,8 @@ void AIAssistantWebApiMessage::Define()
 			It("returns true when the object is identical", [this]
 				{
 					FMessage MessageSame = MakeAgentTextMessage(1, TEXT("Hello"));
+					// Intentional self-compare in test.
+					//-V501
 					TestTrue(TEXT("Message should be equal to itself"), Message == Message);
 					TestTrue(TEXT("Messages created with same helper parameters should be equal"),
 						Message == MessageSame);
@@ -616,6 +628,8 @@ void AIAssistantWebApiConversationId::Define()
 			It("returns true when the object is identical", [this]
 				{
 					FConversationId ConversationIdSame = MakeConversationId(1);
+					// Intentional self-compare in test.
+					//-V501
 					TestTrue(TEXT("ConversationId should be equal to itself"),
 						ConversationId == ConversationId);
 					TestTrue(TEXT("ConversationIds with same Id should be equal"),
@@ -682,7 +696,9 @@ void AIAssistantWebApiConversationDescriptor::Define()
 			It("returns true when descriptors have same values and no ConversationId", [this]
 				{
 					FConversationDescriptor DescriptorSame = Descriptor;
-
+					
+					// Intentional self-compare in test.
+					//-V501
 					TestTrue(TEXT("ConversationDescriptor should be equal to itself"),
 						Descriptor == Descriptor);
 					TestTrue(TEXT("Descriptors with same values and no ConversationId should be equal"),
@@ -797,7 +813,9 @@ void AIAssistantWebApiConversation::Define()
 						MakeAgentTextMessage(1, TEXT("Hello")),
 						MakeUserTextMessage(TEXT("Hi")),
 					};
-
+					
+					// Intentional self-compare in test.
+					//-V501
 					TestTrue(TEXT("Conversation should be equal to itself"),
 						Conversation == Conversation);
 					TestTrue(TEXT("Conversations with same descriptor and messages should be equal"),
@@ -867,7 +885,9 @@ void AIAssistantWebApiAddMessageToConversationOptions::Define()
 					FAddMessageToConversationOptions OptionsSame;
 					OptionsSame.ConversationId = MakeConversationId(1);
 					OptionsSame.Message = MakeAgentTextMessage(1, TEXT("Hello"));
-
+					
+					// Intentional self-compare in test.
+					//-V501
 					TestTrue(TEXT("Options should be equal to itself"), Options == Options);
 					TestTrue(TEXT("Options with same ConversationId and Message should be equal"),
 						Options == OptionsSame);
