[AIAssistant] Second try at fixing static analyzer warnings.

#Synced-CL 48864563

diff --git a/Engine/Plugins/Experimental/AIAssistant/Source/AIAssistant/Private/Tests/AIAssistantWebApiTest.cpp b/Engine/Plugins/Experimental/AIAssistant/Source/AIAssistant/Private/Tests/AIAssistantWebApiTest.cpp
--- a/Engine/Plugins/Experimental/AIAssistant/Source/AIAssistant/Private/Tests/AIAssistantWebApiTest.cpp
+++ b/Engine/Plugins/Experimental/AIAssistant/Source/AIAssistant/Private/Tests/AIAssistantWebApiTest.cpp
@@ -87,8 +87,7 @@ void AIAssistantWebApiTextMessageContent::Define()
 					Same.Text = TEXT("Hello");
 					Same.Format.Reset();
 					// Intentional self-compare in test.
-					//-V501
-					TestTrue(TEXT("Message should be equal to itself"), Message == Message);
+					TestTrue(TEXT("Message should be equal to itself"), Message == Message); //-V501
 					TestTrue(TEXT("Messages with same Text and unset Format should be equal"), Message == Same);
 				});
 			It("returns false when text is different", [this]
@@ -159,8 +158,7 @@ void AIAssistantWebApiToolCallContent::Define()
 					ToolCallSame.ToolCallId.Reset();
 					ToolCallSame.ResponseRequired.Reset();
 					// Intentional self-compare in test.
-					//-V501
-					TestTrue(TEXT("Tool call should be equal to itself"), ToolCall == ToolCall);
+					TestTrue(TEXT("Tool call should be equal to itself"), ToolCall == ToolCall); //-V501
 					TestTrue(TEXT("Tool calls with same fields and unset optionals should be equal"),
 						ToolCall == ToolCallSame);
 				});
@@ -271,8 +269,7 @@ void AIAssistantWebApiToolResponseContent::Define()
 					ResponseSame.ErrorMessage.Reset();
 					
 					// Intentional self-compare in test.
-					//-V501
-					TestTrue(TEXT("Tool response should be equal to itself"), Response == Response);
+					TestTrue(TEXT("Tool response should be equal to itself"), Response == Response); //-V501
 					TestTrue(TEXT("Tool responses with same basic fields and unset optionals should be equal"),
 						Response == ResponseSame);
 				});
@@ -405,8 +402,8 @@ void AIAssistantWebApiMessageContent::Define()
 					TextContentSame.Content.Emplace<FTextMessageContent>(HelloText);
 					
 					// Intentional self-compare in test.
-					//-V501
-					TestTrue(TEXT("Message content should be equal to itself"), TextContent == TextContent);
+					TestTrue(TEXT("Message content should be equal to itself"),
+						TextContent == TextContent); //-V501
 					TestTrue(TEXT("Message contents with same Text variant and visibility should be equal"),
 						TextContent == TextContentSame);
 				});
@@ -472,8 +469,7 @@ void AIAssistantWebApiMessageId::Define()
 				{
 					FMessageId MessageIdSame = MakeMessageId(1);
 					// Intentional self-compare in test.
-					//-V501
-					TestTrue(TEXT("MessageId should be equal to itself"), MessageId == MessageId);
+					TestTrue(TEXT("MessageId should be equal to itself"), MessageId == MessageId); //-V501
 					TestTrue(TEXT("MessageIds created with same index should be equal"),
 						MessageId == MessageIdSame);
 				});
@@ -533,8 +529,7 @@ void AIAssistantWebApiMessage::Define()
 				{
 					FMessage MessageSame = MakeAgentTextMessage(1, TEXT("Hello"));
 					// Intentional self-compare in test.
-					//-V501
-					TestTrue(TEXT("Message should be equal to itself"), Message == Message);
+					TestTrue(TEXT("Message should be equal to itself"), Message == Message); //-V501
 					TestTrue(TEXT("Messages created with same helper parameters should be equal"),
 						Message == MessageSame);
 				});
@@ -629,9 +624,8 @@ void AIAssistantWebApiConversationId::Define()
 				{
 					FConversationId ConversationIdSame = MakeConversationId(1);
 					// Intentional self-compare in test.
-					//-V501
 					TestTrue(TEXT("ConversationId should be equal to itself"),
-						ConversationId == ConversationId);
+						ConversationId == ConversationId); //-V501
 					TestTrue(TEXT("ConversationIds with same Id should be equal"),
 						ConversationId == ConversationIdSame);
 				});
@@ -698,9 +692,8 @@ void AIAssistantWebApiConversationDescriptor::Define()
 					FConversationDescriptor DescriptorSame = Descriptor;
 					
 					// Intentional self-compare in test.
-					//-V501
 					TestTrue(TEXT("ConversationDescriptor should be equal to itself"),
-						Descriptor == Descriptor);
+						Descriptor == Descriptor); //-V501
 					TestTrue(TEXT("Descriptors with same values and no ConversationId should be equal"),
 						Descriptor == DescriptorSame);
 				});
@@ -815,9 +808,8 @@ void AIAssistantWebApiConversation::Define()
 					};
 					
 					// Intentional self-compare in test.
-					//-V501
 					TestTrue(TEXT("Conversation should be equal to itself"),
-						Conversation == Conversation);
+						Conversation == Conversation); //-V501
 					TestTrue(TEXT("Conversations with same descriptor and messages should be equal"),
 						Conversation == ConversationSame);
 				});
@@ -887,8 +879,7 @@ void AIAssistantWebApiAddMessageToConversationOptions::Define()
 					OptionsSame.Message = MakeAgentTextMessage(1, TEXT("Hello"));
 					
 					// Intentional self-compare in test.
-					//-V501
-					TestTrue(TEXT("Options should be equal to itself"), Options == Options);
+					TestTrue(TEXT("Options should be equal to itself"), Options == Options); //-V501
 					TestTrue(TEXT("Options with same ConversationId and Message should be equal"),
 						Options == OptionsSame);
 				});
