AI Assistant - Tool Calling Support

(This will eventually be moved into JsonUtilities.)

#Synced-CL 48840055

diff --git a/Engine/Plugins/Experimental/AIAssistant/Source/AIAssistant/AIAssistant.Build.cs b/Engine/Plugins/Experimental/AIAssistant/Source/AIAssistant/AIAssistant.Build.cs
--- a/Engine/Plugins/Experimental/AIAssistant/Source/AIAssistant/AIAssistant.Build.cs
+++ b/Engine/Plugins/Experimental/AIAssistant/Source/AIAssistant/AIAssistant.Build.cs
@@ -48,6 +48,7 @@ public class AIAssistant : ModuleRules
 				"DeveloperSettings",
 				"Json",
 				"JsonUtilities",
+				"JsonUtilitiesEditor",
 				"EditorSubsystem",
 				"HTTP",
                 "EditorScriptingUtilities",
diff --git a/Engine/Plugins/Experimental/AIAssistant/Source/AIAssistant/Private/AIAssistant.cpp b/Engine/Plugins/Experimental/AIAssistant/Source/AIAssistant/Private/AIAssistant.cpp
--- a/Engine/Plugins/Experimental/AIAssistant/Source/AIAssistant/Private/AIAssistant.cpp
+++ b/Engine/Plugins/Experimental/AIAssistant/Source/AIAssistant/Private/AIAssistant.cpp
@@ -18,6 +18,7 @@
 #include "AIAssistantSlateQuerier.h"
 #include "AIAssistantStyle.h"
 #include "AIAssistantWebBrowser.h"
+#include "AIAssistantToolCall.h"
 
 #define LOCTEXT_NAMESPACE "FAIAssistantModule"
 
@@ -44,11 +45,13 @@ static const FName AIAssistantTabName("AIAssistant");
 		return;
 	}
 	
-		
+
 	FAIAssistantStyle::Initialize();
 	FAIAssistantStyle::ReloadTextures();
 
 
+	UE::AIAssistant::FObjectFunctionToolCall::Initialize();
+	
 	FAIAssistantCommands::Register();
 
 	PluginCommands = MakeShareable(new FUICommandList);
diff --git a/Engine/Plugins/Experimental/AIAssistant/Source/AIAssistant/Private/AIAssistantToolCall.cpp b/Engine/Plugins/Experimental/AIAssistant/Source/AIAssistant/Private/AIAssistantToolCall.cpp
new file mode 100644
--- /dev/null
+++ b/Engine/Plugins/Experimental/AIAssistant/Source/AIAssistant/Private/AIAssistantToolCall.cpp
@@ -0,0 +1,188 @@
+// Copyright Epic Games, Inc. All Rights Reserved.
+
+#include "AIAssistantToolCall.h"
+
+#include "Engine/Engine.h"
+#include "Engine/World.h"
+#include "Engine/EngineTypes.h"
+#include "Modules/ModuleManager.h"
+#include "JsonDomBuilder.h"
+#include "Async/Async.h"
+#include "Serialization/JsonReader.h"
+#include "Serialization/JsonSerializer.h"
+#include "UObject/FieldIterator.h"
+#include "JsonSchema/JsonSchemaGenerator.h"
+
+#define LOCTEXT_NAMESPACE "FObjectFunctionToolCall"
+
+PRAGMA_DISABLE_DEPRECATION_WARNINGS
+
+using namespace UE::AIAssistant;
+
+//
+// Internal
+//
+
+namespace
+{
+	TSharedPtr<FJsonObject> JsonStringToJsonObject(const FString& JsonString)
+	{
+		if (!JsonString.IsEmpty())
+		{
+			const TSharedRef<TJsonReader<>> JsonReader = TJsonReaderFactory<>::Create(JsonString);
+			TSharedPtr<FJsonObject> JsonObject;
+			if (FJsonSerializer::Deserialize(JsonReader, JsonObject))
+			{
+				return JsonObject;
+			}
+		}
+		
+		return nullptr;			
+	}
+	
+	TSharedPtr<FJsonObject> FunctionInputParamsJsonToJsonObject(
+		const FObjectFunctionToolCall::FFunctionInputParamsJson& FunctionInputParamsJson)
+	{
+		if (FunctionInputParamsJson.IsType<TSharedPtr<FJsonObject>>())
+		{
+			return FunctionInputParamsJson.Get<TSharedPtr<FJsonObject>>();
+		}
+		else
+		{
+			return JsonStringToJsonObject(FunctionInputParamsJson.Get<FString>());
+		}
+	}
+	
+	UWorld* GetWorld()
+	{
+		if (GEngine)
+		{
+			for (const FWorldContext& WorldContext : GEngine->GetWorldContexts())
+			{
+				if (WorldContext.WorldType == EWorldType::Game || WorldContext.WorldType == EWorldType::PIE)
+				{
+					return WorldContext.World();
+				}
+			}
+		}
+		
+		return nullptr;
+	}
+	
+	void MaybeSetWorldContextPropertyInFrameMemory(const UFunction* Function, uint8* FrameMemory)
+	{
+		const FName WorldContextKey(TEXT("WorldContext"));
+		const FName WorldContextParamName = 
+			Function->HasMetaData(WorldContextKey) ?
+			FName(*Function->GetMetaData(WorldContextKey)) :
+			FName(TEXT("WorldContextObject"));
+		
+		if (FProperty* WorldContextProperty = Function->FindPropertyByName(WorldContextParamName))
+		{
+			UWorld* World = GetWorld();
+			FObjectPropertyBase* ObjectProperty = CastField<FObjectPropertyBase>(WorldContextProperty);
+			if (ObjectProperty && World)
+			{
+				ObjectProperty->SetObjectPropertyValue_InContainer(FrameMemory, World);
+			}
+		}
+	}
+}
+
+//
+// AI::AIAssistant::FObjectFunctionToolCall
+//
+
+/*static*/ void FObjectFunctionToolCall::Initialize()
+{
+	FModuleManager::Get().LoadModuleChecked(TEXT("JsonUtilitiesEditor"));
+}
+
+FObjectFunctionToolCall::FObjectFunctionToolCall(TNotNull<UObject*> InInstanceObject, TNotNull<UFunction*> InFunction) :
+	InstanceObject(InInstanceObject), 
+	Function(InFunction)
+{
+}
+
+FObjectFunctionToolCall::FToolCallResult FObjectFunctionToolCall::Execute(const FFunctionInputParamsJson& FunctionInputParamsJson) const
+{
+	// Get incoming function params JSON object.
+	const TSharedPtr<FJsonObject> FunctionInputParamsJsonObject = 
+		FunctionInputParamsJsonToJsonObject(FunctionInputParamsJson);
+	
+	// Allocate the frame memory.
+	uint8* FrameMemory;
+	{
+		// Invoke must use properties size, because it's allocating locals for e.g. script functions.
+		const int32 FrameSize = Function->PropertiesSize; 
+		FrameMemory = static_cast<uint8*>(FMemory_Alloca_Aligned(FrameSize, Function->GetMinAlignment()));
+		FMemory::Memzero(FrameMemory, FrameSize);
+	}
+	
+	// We need to manually allocate each property in the frame memory.
+	for (TFieldIterator<FProperty> It(Function); It && It->HasAnyPropertyFlags(CPF_Parm); ++It)
+	{
+		if (const FProperty* ParameterProp = *It;
+			!ParameterProp->HasAnyPropertyFlags(CPF_ZeroConstructor))
+		{
+			ParameterProp->InitializeValue_InContainer(FrameMemory);
+		}
+	}
+
+	// When we exit, we need to manually destroy each property in the frame memory.
+	ON_SCOPE_EXIT
+	{
+		// Invoke will destroy locals, but caller owns parameters. This should also destroy the result parameter.
+		for (TFieldIterator<FProperty> It(Function); It && It->HasAnyPropertyFlags(CPF_Parm); ++It)
+		{
+			It->DestroyValue_InContainer(FrameMemory);
+		}
+	};
+
+	// Copy the function input params memory into the function's frame memory.
+	if (FunctionInputParamsJsonObject)
+	{
+		const bool bDidConvertObjectToUStruct = FJsonObjectConverter::JsonObjectToUStruct(FunctionInputParamsJsonObject.ToSharedRef(), Function, 
+			FrameMemory);
+		check(bDidConvertObjectToUStruct);
+		if (!bDidConvertObjectToUStruct)
+		{
+			static const FText ErrorString = LOCTEXT("JsonObjectToUStructFailed", "Could not convert input params JSON object to a UStruct.");
+			return MakeError(ErrorString);
+		}
+	}
+	
+	// In case this is a Blueprint function that has a world context object parameter, 
+	// give that parameter (in the frame memory) the World.
+	MaybeSetWorldContextPropertyInFrameMemory(Function, FrameMemory);
+
+	// Invoke function.
+	FFrame StackFrame(InstanceObject, Function, FrameMemory, nullptr, Function->ChildProperties);
+	const bool bHasReturnParam = (Function->ReturnValueOffset != MAX_uint16);
+	uint8* ReturnValueAddress = (bHasReturnParam ? (FrameMemory + Function->ReturnValueOffset) : nullptr);
+	Function->Invoke(InstanceObject, StackFrame, ReturnValueAddress);
+	
+	// Get the return value, if we have one.
+	TSharedPtr<FJsonObject> ReturnJsonObject; 
+	if (FProperty* ReturnProperty = Function->GetReturnProperty())
+	{
+		if (const TSharedPtr<FJsonValue> ReturnJsonValue = FJsonObjectConverter::UPropertyToJsonValue(ReturnProperty, ReturnValueAddress);
+			ReturnJsonValue->Type == EJson::Object)
+		{
+			// If value is an object, return it.
+			ReturnJsonObject = ReturnJsonValue->AsObject();
+		}
+		else
+		{
+			// If value is not an object, return it inside an object.
+			ReturnJsonObject = MakeShared<FJsonObject>();
+			ReturnJsonObject->SetField(TEXT("returnValue"), ReturnJsonValue);
+		}
+	}
+
+	return MakeValue(ReturnJsonObject);
+}
+
+PRAGMA_ENABLE_DEPRECATION_WARNINGS
+
+#undef LOCTEXT_NAMESPACE
diff --git a/Engine/Plugins/Experimental/AIAssistant/Source/AIAssistant/Private/AIAssistantToolCall.h b/Engine/Plugins/Experimental/AIAssistant/Source/AIAssistant/Private/AIAssistantToolCall.h
new file mode 100644
--- /dev/null
+++ b/Engine/Plugins/Experimental/AIAssistant/Source/AIAssistant/Private/AIAssistantToolCall.h
@@ -0,0 +1,80 @@
+// Copyright Epic Games, Inc. All Rights Reserved.
+
+#pragma once
+
+#include "Templates/SharedPointer.h"
+#include "Misc/NotNull.h"
+#include "Misc/TVariant.h"
+#include "Templates/ValueOrError.h"
+#include "UObject/Class.h"
+#include "Dom/JsonObject.h"
+
+//
+// AI::AIAssistant::FObjectFunctionToolCall
+//
+
+namespace UE::AIAssistant
+{
+	class FObjectFunctionToolCall
+	{
+	public:
+
+		/**
+		 * Holds either a JSON object or a JSON string.
+		 * Usage -
+		 *		FFunctionInputParamsJson FunctionInputParamsJson;
+		 *		...
+		 *		FunctionInputParamsJson.Set<FString>(SomeJsonString);
+		 *		or
+		 *		FunctionInputParamsJson.Set<TSharedPtr<FJsonObject>>(SomeJsonSharedObjectPtr);
+		 */
+		using FFunctionInputParamsJson = TVariant<TSharedPtr<FJsonObject>, FString>;
+
+		/**
+		 * Tool call result.
+		 * Either JSON object or text describing the error.
+		 */
+		using FToolCallResult = TValueOrError<const TSharedPtr<FJsonObject>, const FText>;
+		
+		/**
+		 * Initializes use of this class.
+		 */
+		static void Initialize();
+		
+		/**
+		 * Constructs with a UFunction on a UObject instance. 
+		 * @param InInstanceObject UObject to call the function on.
+		 * @param InFunction UFunction to call on the UObject.
+		 */
+		FObjectFunctionToolCall(TNotNull<UObject*> InInstanceObject, TNotNull<UFunction*> InFunction);
+		
+		/**
+		 * Constructs with a UFunction on a UObject instance, but specialized for UObjects that are UClasses where the function can be found by name.
+		 * @param InInstanceObject UObject to call the function on.
+		 * @param InFunctionName UFunction name to find and call on the UObject.
+		 */
+		static TOptional<FObjectFunctionToolCall> FromNamedClassFunction(TNotNull<UObject*> InInstanceObject, const FName& InFunctionName)
+		{
+			check(!InFunctionName.IsNone());
+		
+			UFunction* Function = InInstanceObject->GetClass()->FindFunctionByName(InFunctionName);
+			check(Function);
+
+			return TOptional(FObjectFunctionToolCall(InInstanceObject, Function));
+		}
+		
+		/**
+		 * Given a JSON containing input parameters for a tool call, execute a function with those input parameters.
+		 * @param FunctionInputParamsJson The JSON object or string containing the function parameters.
+		 * @return Tool call result. (See above). When tool result does not have an error -
+ 		 *		Returned objects will either be (1) the object itself, or (2) a value named 'returnValue' inside an object.
+		 *		May be null if (1) function does not return a value, or (2) if function does return a value and that return value was null.
+		 */
+		FToolCallResult Execute(const FFunctionInputParamsJson& FunctionInputParamsJson) const;
+		
+	private:
+		
+		TNotNull<UObject*> InstanceObject;
+		TNotNull<UFunction*> Function;
+	};
+}
diff --git a/Engine/Plugins/Experimental/AIAssistant/Source/AIAssistant/Private/Tests/AIAssistantToolCallTest.cpp b/Engine/Plugins/Experimental/AIAssistant/Source/AIAssistant/Private/Tests/AIAssistantToolCallTest.cpp
new file mode 100644
--- /dev/null
+++ b/Engine/Plugins/Experimental/AIAssistant/Source/AIAssistant/Private/Tests/AIAssistantToolCallTest.cpp
@@ -0,0 +1,118 @@
+// Copyright Epic Games, Inc. All Rights Reserved.
+
+#include "AIAssistantToolCallTest.h"
+
+#include "Misc/AutomationTest.h"
+#include "UObject/Package.h"
+
+#include "AIAssistantLog.h"
+#include "Tests/AIAssistantTestFlags.h"
+#include "JsonSchema/JsonSchemaGenerator.h"
+
+#include "AIAssistantToolCall.h"
+
+
+using namespace UE::AIAssistant;
+	
+
+PRAGMA_DISABLE_DEPRECATION_WARNINGS
+
+
+#define PRINT_JSONS_TO_LOG 1
+
+static void MaybePrintJsonToLog(const TCHAR* Name, const TSharedRef<FJsonObject>& JsonObject)
+{
+#if PRINT_JSONS_TO_LOG
+	
+	FString JsonString;
+	const TSharedRef<TJsonWriter<>> JsonWriter = TJsonWriterFactory<>::Create(&JsonString);
+	FJsonSerializer::Serialize(JsonObject, JsonWriter);
+	UE_LOG(LogAIAssistant, Warning, TEXT("[%s]\n%s"), Name, *JsonString);
+	
+#endif
+}
+
+
+#if WITH_DEV_AUTOMATION_TESTS
+
+
+IMPLEMENT_SIMPLE_AUTOMATION_TEST(
+	FAIAssistantToolCallExecuteToolCallFunction,
+	"AI.Assistant.ToolCall.ExecuteToolCallFunction",
+	AIAssistantTest::Flags);
+
+bool FAIAssistantToolCallExecuteToolCallFunction::RunTest(const FString& Parameters)
+{
+	static const FName FunctionName = "TestFunc";
+	static constexpr float FunctionInputFloatValue = 3.0f;
+	static constexpr float ClassFloatValue = 20.0f;
+	static const FString FunctionInputParamsJsonString = FString::Printf(TEXT("{\n\"InFloat\": %f\n}"), FunctionInputFloatValue);
+
+	// Create an object and provide it specific internal state. 
+	// (For this test - Using TStrongObjectPtr to prevent GC.)
+	
+	const TStrongObjectPtr InstanceObject = TStrongObjectPtr(NewObject<UToolCallFakeClass>());
+
+	InstanceObject->Float = ClassFloatValue;
+	
+	// Execute.
+	
+	const TOptional<FObjectFunctionToolCall> ToolCall = FObjectFunctionToolCall::FromNamedClassFunction(InstanceObject.Get(), FunctionName);
+	if (!TestTrue("Tool call created from class and function name.", ToolCall.IsSet()))
+	{
+		return false;
+	}
+	
+	FObjectFunctionToolCall::FFunctionInputParamsJson JsonInputParamsJson;
+	JsonInputParamsJson.Set<FString>(FunctionInputParamsJsonString);
+	const FObjectFunctionToolCall::FToolCallResult ToolCallResult = ToolCall->Execute(JsonInputParamsJson);
+	
+	if (!TestTrue("Execute tool call result does not have error.", !ToolCallResult.HasError()))
+	{
+		UE_LOG(LogAIAssistant, Error, TEXT("Tool call error - %s"), *ToolCallResult.GetError().ToString());
+		return false;
+	}
+	if (!TestTrue("Execute tool call result has value.", ToolCallResult.HasValue()))
+	{
+		return false;
+	}
+
+	const TSharedPtr<FJsonObject> JsonObject = ToolCallResult.GetValue();
+
+	MaybePrintJsonToLog(TEXT("ToolCallReturnedJsonObject"), JsonObject.ToSharedRef());
+	
+	if (!TestTrue("Execute tool call result has valid value.", JsonObject.IsValid()))
+	{
+		return false;
+	}
+	if (!TestTrue("Return object is not empty (has values.)", JsonObject->Values.Num() > 0))
+	{
+		return false;
+	}	
+	if (!TestTrue(TEXT("Return object has float returnValue."), JsonObject->HasTypedField(TEXT("returnValue"), EJson::Number)))
+	{
+		return false;
+	}
+	
+	// Check.
+	
+	const float ReturnValueFloat = JsonObject->GetNumberField(TEXT("returnValue"));
+	const float ReturnValueFloatExpected = InstanceObject->TestFunc(FunctionInputFloatValue);
+	const FString TestMessage = FString::Printf(TEXT("Return value %f is expected value %f."), ReturnValueFloat, ReturnValueFloatExpected);
+	if (!TestTrue(TestMessage, ReturnValueFloat == ReturnValueFloatExpected))
+	{
+		return false;
+	}
+	
+	// Success.
+	
+	return true;
+}
+
+
+#endif
+
+
+PRAGMA_ENABLE_DEPRECATION_WARNINGS
+
+
diff --git a/Engine/Plugins/Experimental/AIAssistant/Source/AIAssistant/Private/Tests/AIAssistantToolCallTest.h b/Engine/Plugins/Experimental/AIAssistant/Source/AIAssistant/Private/Tests/AIAssistantToolCallTest.h
new file mode 100644
--- /dev/null
+++ b/Engine/Plugins/Experimental/AIAssistant/Source/AIAssistant/Private/Tests/AIAssistantToolCallTest.h
@@ -0,0 +1,24 @@
+// Copyright Epic Games, Inc. All Rights Reserved.
+
+#pragma once
+
+#include "CoreMinimal.h"
+
+#include "AIAssistantToolCallTest.generated.h"
+
+UCLASS(BlueprintType)
+class UToolCallFakeClass : public UObject
+{
+	GENERATED_BODY()
+	
+public:
+	
+	UPROPERTY(EditAnywhere, Category="Main")
+	float Float = 0.0f;
+	
+	UFUNCTION(BlueprintCallable, Category="Main") 
+	float TestFunc(const float InFloat) const 
+	{
+		return 100.0f + Float + InFloat; 
+	}
+};
