[AIAssistant] Add GetConversation to the WebApi

#Synced-CL 48864563

diff --git a/Engine/Plugins/Experimental/AIAssistant/Source/AIAssistant/Private/AIAssistantWebApi.cpp b/Engine/Plugins/Experimental/AIAssistant/Source/AIAssistant/Private/AIAssistantWebApi.cpp
--- a/Engine/Plugins/Experimental/AIAssistant/Source/AIAssistant/Private/AIAssistantWebApi.cpp
+++ b/Engine/Plugins/Experimental/AIAssistant/Source/AIAssistant/Private/AIAssistantWebApi.cpp
@@ -83,6 +83,21 @@ try {
 		return ExecutionFunctionParseJson<void>(*WebApiObjectName, TEXT("createConversation"));
 	}
 
+	TFuture<TValueOrError<FConversation, FString>> FWebApi::GetConversation(
+		const TOptional<FConversationId>& ConversationId)
+	{
+		if (ConversationId.IsSet())
+		{
+			return ExecutionFunctionParseJson<FConversation>(
+				*WebApiObjectName, TEXT("getConversation"), *ConversationId);
+		}
+		else
+		{
+			return ExecutionFunctionParseJson<FConversation>(
+				*WebApiObjectName, TEXT("getConversation"), nullptr);
+		}
+	}
+
 	TFuture<TValueOrError<FAgentEnvironmentHandle, FString>> FWebApi::AddAgentEnvironment(
 		const FAgentEnvironment& AgentEnvironment)
 	{
diff --git a/Engine/Plugins/Experimental/AIAssistant/Source/AIAssistant/Private/AIAssistantWebApi.h b/Engine/Plugins/Experimental/AIAssistant/Source/AIAssistant/Private/AIAssistantWebApi.h
--- a/Engine/Plugins/Experimental/AIAssistant/Source/AIAssistant/Private/AIAssistantWebApi.h
+++ b/Engine/Plugins/Experimental/AIAssistant/Source/AIAssistant/Private/AIAssistantWebApi.h
@@ -237,14 +237,16 @@ namespace UE::AIAssistant
 	{
 		// Unique ID of the conversation.
 		FString Id;
+		FString Type;
 
 		bool operator==(const FConversationId& Other) const
 		{
-			return Id == Other.Id;
+			return Id == Other.Id && Type == Other.Type;
 		}
 
 		BEGIN_JSON_SERIALIZER
 			JSON_SERIALIZE("id", Id);
+			JSON_SERIALIZE("type", Type);
 		END_JSON_SERIALIZER
 	};
 
@@ -436,6 +438,11 @@ namespace UE::AIAssistant
 		// Create a new conversation.
 		TFuture<TValueOrError<void, FString>> CreateConversation();
 
+		// Get a conversation by ID.
+		// If no ID is specified, the current conversation should be returned.
+		virtual TFuture<TValueOrError<FConversation, FString>> GetConversation(
+			const TOptional<FConversationId>& ConversationId);
+
 		// Add an agent environment for the currently logged in user
 		// returning the ID. If a matching environment already exists for the
 		// user, this should return the existing environment (i.e upsert).
@@ -498,6 +505,14 @@ namespace UE::AIAssistant
 		{
 			return ExecuteFunction(InstanceName, FunctionName, *Argument.ToJson(false));
 		}
+		
+		// Execute a javascript function converting an empty argument to JSON null.
+		TFuture<UAIAssistantWebJavaScriptResultDelegate::FResult> ExecuteFunctionWithJsonArgument(
+			const TCHAR* InstanceName, const TCHAR* FunctionName,
+			std::nullptr_t NullPtr)
+		{
+			return ExecuteFunction(InstanceName, FunctionName, TEXT("null"));
+		}
 
 		// Create a promise and handler for the execution of a JavaScript function that optionally
 		// parses a JSON result value.
@@ -608,7 +623,7 @@ namespace UE::AIAssistant
 
 FORCEINLINE uint32 GetTypeHash(const UE::AIAssistant::FConversationId& Id)
 {
-	return GetTypeHash(Id.Id);
+	return HashCombine(GetTypeHash(Id.Id), GetTypeHash(Id.Type));
 }
 
 FORCEINLINE uint32 GetTypeHash(const UE::AIAssistant::FMessageId& Id)
diff --git a/Engine/Plugins/Experimental/AIAssistant/Source/AIAssistant/Private/Tests/AIAssistantMessageUtils.cpp b/Engine/Plugins/Experimental/AIAssistant/Source/AIAssistant/Private/Tests/AIAssistantMessageUtils.cpp
--- a/Engine/Plugins/Experimental/AIAssistant/Source/AIAssistant/Private/Tests/AIAssistantMessageUtils.cpp
+++ b/Engine/Plugins/Experimental/AIAssistant/Source/AIAssistant/Private/Tests/AIAssistantMessageUtils.cpp
@@ -29,6 +29,7 @@ namespace UE::AIAssistant
 	{
 		FConversationId ConversationId;
 		ConversationId.Id = FString::Printf(TEXT("Conversation%i"), Index);
+		ConversationId.Type = TEXT("ConversationId");
 		return ConversationId;
 	}
 
diff --git a/Engine/Plugins/Experimental/AIAssistant/Source/AIAssistant/Private/Tests/AIAssistantWebApiTest.cpp b/Engine/Plugins/Experimental/AIAssistant/Source/AIAssistant/Private/Tests/AIAssistantWebApiTest.cpp
--- a/Engine/Plugins/Experimental/AIAssistant/Source/AIAssistant/Private/Tests/AIAssistantWebApiTest.cpp
+++ b/Engine/Plugins/Experimental/AIAssistant/Source/AIAssistant/Private/Tests/AIAssistantWebApiTest.cpp
@@ -1132,6 +1132,37 @@ bool FAIAssistantWebApiTestCreateConversation::RunTest(const FString& UnusedPara
 		*this, nullptr, TEXT("createConversation"), TEXT(""), Result, TEXT(""), false);
 }
 
+BEGIN_DEFINE_SPEC(FAIAssistantWebApiTest, "AI.Assistant.WebApi", AIAssistantTest::Flags);
+	FConversation Conversation;
+END_DEFINE_SPEC(FAIAssistantWebApiTest)
+
+void FAIAssistantWebApiTest::Define()
+{
+	Describe("GetConversation", [this]
+		{
+			It("passes the ConversationId to getConversation when supplied", [this]
+				{
+					FFakeWebApiContainer WebApi;
+					TOptional<FConversationId> ConversationId = MakeConversationId(123);
+					auto Result = WebApi->GetConversation(ConversationId);
+					return WebApi->TestExpectAsyncFunctionCallAndComplete(
+						*this, nullptr, TEXT("getConversation"),
+						TEXT(R"json({"id":"Conversation123","type":"ConversationId"})json"),
+						Result, *Conversation.ToJson(false), false);
+				});
+			It("passes null to getConversation when the optional is a nullopt", [this]
+				{
+					FFakeWebApiContainer WebApi;
+					TOptional<FConversationId> ConversationId;
+					auto Result = WebApi->GetConversation(ConversationId);
+					return WebApi->TestExpectAsyncFunctionCallAndComplete(
+						*this, nullptr, TEXT("getConversation"),
+						TEXT(R"json(null)json"), Result, *Conversation.ToJson(false), false);
+				});
+		});
+}
+
+
 IMPLEMENT_SIMPLE_AUTOMATION_TEST(
 	FAIAssistantWebApiTestAddMessageToConversation,
 	"AI.Assistant.WebApi.AddMessageToConversation",
@@ -1458,4 +1489,4 @@ bool FAIAssistantWebApiTestUnregisterOnConversationUpdate::RunTest(const FString
 		&& TestEqual(TEXT("UnregisterOnConversationUpdate_CallCount"), CallCount, 1);
 }
 
-#endif  // WITH_DEV_AUTOMATION_TESTS
\ No newline at end of file
+#endif  // WITH_DEV_AUTOMATION_TESTS
