[AI Assistant] Implementation of UAIAssistantToolset base class

Provides base class for toolsets to derive from implementing IToolsetHandler.
Derived toolsets can be implemented in C++, Blueprint, or Python.

#rnx
#Synced-CL 48944728

diff --git a/Engine/Plugins/Experimental/AIAssistant/Source/AIAssistant/Private/AIAssistantBPToolsetRegistry.h b/Engine/Plugins/Experimental/AIAssistant/Source/AIAssistant/Private/AIAssistantBPToolsetRegistry.h
new file mode 100644
--- /dev/null
+++ b/Engine/Plugins/Experimental/AIAssistant/Source/AIAssistant/Private/AIAssistantBPToolsetRegistry.h
@@ -0,0 +1,85 @@
+// Copyright Epic Games, Inc. All Rights Reserved.
+
+#pragma once
+
+#include "CoreMinimal.h"
+#include "Async/Future.h"
+#include "Kismet/BlueprintFunctionLibrary.h"
+#include "Editor.h"
+#include "Misc/Variant.h"
+#include "Templates/SharedPointer.h"
+
+#include "AIAssistantLog.h"
+#include "AIAssistantSubsystem.h"
+#include "AIAssistantToolsetRegistry.h"
+#include "AIAssistantToolsetHandler.h"
+
+#include "AIAssistantBPToolsetRegistry.generated.h"
+
+// Blueprint-accessible wrapper to allow queries against the AI Assistant Toolset Registry.
+UCLASS()
+class AIASSISTANT_API UAIAssistantToolsetRegistry : public UBlueprintFunctionLibrary
+{
+	GENERATED_BODY()
+
+public:
+
+	UFUNCTION(BlueprintPure, Category="ToolsetRegistry")
+	static bool IsToolsetRegistered(const FString& ToolsetName)
+	{
+		auto AssistantSubsystem = UAIAssistantSubsystem::Get(TEXT("IsToolsetRegistered"));
+		return AssistantSubsystem.HasValue() &&
+			static_cast<bool>(AssistantSubsystem.GetValue()->ToolsetRegistry.Find(ToolsetName));
+	}
+
+	UFUNCTION(BlueprintCallable, Category="ToolsetRegistry")
+	static void UnregisterToolsetByName(const FString& ToolsetName)
+	{
+		auto AssistantSubsystem = UAIAssistantSubsystem::Get(TEXT("UnregisterToolsetByName"));
+		if (AssistantSubsystem.HasValue())
+		{
+			TSharedPtr<UE::AIAssistant::IToolsetHandler>* PtrToHandler =
+				AssistantSubsystem.GetValue()->ToolsetRegistry.Find(ToolsetName);
+			if (PtrToHandler)
+			{
+				TSharedPtr<UE::AIAssistant::IToolsetHandler> Handler = *PtrToHandler;
+				UE_LOG(LogAIAssistant, Display, TEXT("Unregistering Toolset %s"), *ToolsetName);
+				AssistantSubsystem.GetValue()->ToolsetRegistry.UnregisterToolset(Handler);
+			}
+		}
+	}
+
+	UFUNCTION(BlueprintCallable, Category="ToolsetRegistry")
+	static void ExecuteTool(const FString& ToolsetName, const FString& ToolName, const FString& JsonInput,
+							FString& OutResult, FString& OutError)
+	{
+		OutResult.Empty();
+		OutError.Empty();
+
+		auto AssistantSubsystem = UAIAssistantSubsystem::Get(TEXT("ExecuteTool"));
+		if (AssistantSubsystem.HasError())
+		{
+			OutError = AssistantSubsystem.StealError();
+			return;
+		}
+
+		if (AssistantSubsystem.HasValue())
+		{
+			UE::AIAssistant::FToolDescriptor ToolDescriptor;
+			ToolDescriptor.ToolsetName = ToolsetName;
+			ToolDescriptor.ToolName = ToolName;
+			TFuture<TValueOrError<FString, FString>> Future =
+				AssistantSubsystem.GetValue()->ToolsetRegistry.ExecuteTool(ToolDescriptor, JsonInput);
+			// TODO: Block for now, will do proper async handling soon.
+			TValueOrError<FString, FString> Result = Future.Get();
+			if (Result.HasError())
+			{
+				OutError = Result.StealError();
+			}
+			else if (Result.HasValue())
+			{
+				OutResult = Result.StealValue();
+			}
+		}
+	}
+};
diff --git a/Engine/Plugins/Experimental/AIAssistant/Source/AIAssistant/Private/AIAssistantSubsystem.cpp b/Engine/Plugins/Experimental/AIAssistant/Source/AIAssistant/Private/AIAssistantSubsystem.cpp
--- a/Engine/Plugins/Experimental/AIAssistant/Source/AIAssistant/Private/AIAssistantSubsystem.cpp
+++ b/Engine/Plugins/Experimental/AIAssistant/Source/AIAssistant/Private/AIAssistantSubsystem.cpp
@@ -1,13 +1,16 @@
 // Copyright Epic Games, Inc. All Rights Reserved.
 
-
 #include "AIAssistantSubsystem.h"
 
 #include "Containers/UnrealString.h"
 #include "Editor.h"
 #include "Modules/ModuleManager.h"
+#include "Templates/SharedPointer.h"
+#include "Templates/ValueOrError.h"
+#include "UObject/ObjectPtr.h"
 
 #include "AIAssistant.h"
+#include "AIAssistantLog.h"
 #include "AIAssistantPythonExecutor.h"
 #include "AIAssistantWebBrowser.h"
 
@@ -21,6 +24,35 @@ using namespace UE::AIAssistant::PythonExecutor;
 // Statics.
 //
 
+TValueOrError<TObjectPtr<UAIAssistantSubsystem>, FString> UAIAssistantSubsystem::Get(
+	const FString& WarningMessage)
+{
+	// Local function to construct error to return and optionally log a warning.
+	auto MaybeWarnOnError = [&WarningMessage](const FString& ErrorMessage)
+	{
+		if (!WarningMessage.IsEmpty())
+		{
+			UE_LOG(LogAIAssistant, Warning, TEXT("%s: %s"), *WarningMessage, *ErrorMessage);
+		}
+		return MakeError(ErrorMessage);
+	};
+
+	if (!GEditor)
+	{
+		return MaybeWarnOnError(TEXT("Editor is not available"));
+	}
+
+	TObjectPtr<UAIAssistantSubsystem> AssistantSubsystem(
+		GEditor->GetEditorSubsystem<UAIAssistantSubsystem>());
+
+	if (!AssistantSubsystem)
+	{
+		return MaybeWarnOnError(TEXT("AI Assistant subsystem is not available"));
+	}
+
+	return MakeValue(AssistantSubsystem);
+}
+
 FAIAssistantModule& UAIAssistantSubsystem::GetAIAssistantModule()
 {
 	return FModuleManager::LoadModuleChecked<FAIAssistantModule>(UE_PLUGIN_NAME);
@@ -57,7 +89,7 @@ void UAIAssistantSubsystem::Deinitialize()
 	{
 		CodeOutput = TEXT("Code executed successfully.");
 	}
-	
+
 	return CodeOutput;
 }
 
diff --git a/Engine/Plugins/Experimental/AIAssistant/Source/AIAssistant/Private/AIAssistantSubsystem.h b/Engine/Plugins/Experimental/AIAssistant/Source/AIAssistant/Private/AIAssistantSubsystem.h
--- a/Engine/Plugins/Experimental/AIAssistant/Source/AIAssistant/Private/AIAssistantSubsystem.h
+++ b/Engine/Plugins/Experimental/AIAssistant/Source/AIAssistant/Private/AIAssistantSubsystem.h
@@ -5,8 +5,11 @@
 
 #include "Containers/UnrealString.h"
 #include "EditorSubsystem.h"
-#include "AIAssistantToolsetRegistry.h"
 #include "Templates/SharedPointer.h"
+#include "Templates/ValueOrError.h"
+#include "UObject/ObjectPtr.h"
+
+#include "AIAssistantToolsetRegistry.h"
 
 #include "AIAssistantSubsystem.generated.h"
 
@@ -41,6 +44,10 @@ public:
 	void ShowContextMenuViaJavaScript(const FString& SelectedString, const int32 ClientX, const int32 ClientY) const;
 
 public:
+	// Get the assistant subsystem.  Optionally log a warning if passed a non-empty string (the warning
+	// message will start with the passed-in string).
+	static TValueOrError<TObjectPtr<UAIAssistantSubsystem>, FString> Get(const FString& WarningMessage=FString());
+
 	// Get the assistant module.
 	static FAIAssistantModule& GetAIAssistantModule();
 	// Get the assistant web browser.
diff --git a/Engine/Plugins/Experimental/AIAssistant/Source/AIAssistant/Private/AIAssistantToolset.cpp b/Engine/Plugins/Experimental/AIAssistant/Source/AIAssistant/Private/AIAssistantToolset.cpp
new file mode 100644
--- /dev/null
+++ b/Engine/Plugins/Experimental/AIAssistant/Source/AIAssistant/Private/AIAssistantToolset.cpp
@@ -0,0 +1,322 @@
+// Copyright Epic Games, Inc. All Rights Reserved.
+
+#include "AIAssistantToolset.h"
+
+#include "Async/Future.h"
+#include "Containers/UnrealString.h"
+#include "Dom/JsonObject.h"
+#include "Editor.h"
+#include "JsonSchema/JsonSchemaGenerator.h"
+#include "Policies/CondensedJsonPrintPolicy.h"
+#include "Serialization/JsonWriter.h"
+#include "Serialization/JsonSerializer.h"
+#include "Templates/ValueOrError.h"
+#include "UObject/Class.h"
+
+#include "AIAssistantLog.h"
+#include "AIAssistantSubsystem.h"
+#include "AIAssistantToolCall.h"
+#include "AIAssistantToolsetHandler.h"
+
+
+namespace UE::AIAssistant
+{
+
+	/**
+	 * Implements IToolsetHandler for registration.
+	 *
+	 * Hands off calls from AI Assistant registry to UAIAssistantToolset.
+	 */
+	class FToolsetHandlerTrampoline : public UE::AIAssistant::IToolsetHandler
+	{
+	  public:
+		FToolsetHandlerTrampoline(TWeakObjectPtr<UAIAssistantToolset> InToolsetHandler)
+			: ToolsetHandler(InToolsetHandler)
+		{
+			TStrongObjectPtr<UAIAssistantToolset> Handler = ToolsetHandler.Pin();
+			if (Handler)
+			{
+				ToolsetName = Handler->GetToolsetName();
+				ToolsetVersion = Handler->GetToolsetVersion();
+				ToolsetSchema = Handler->GetToolsetJsonSchema();
+			}
+		}
+
+		// Executes a tool with the given name and JSON input, returning a future with the JSON
+		// output.
+		virtual TFuture<TValueOrError<FString, FString>> ExecuteTool(
+			const FString& ToolName, const FString& JsonInput) override
+		{
+			FString ToolResult;
+			FString ToolError;
+			TStrongObjectPtr<UAIAssistantToolset> Handler = ToolsetHandler.Pin();
+			if (Handler)
+			{
+				Handler->InvokeTool(ToolName, JsonInput, ToolResult, ToolError);
+				if (ToolError.IsEmpty())
+				{
+					return MakeFulfilledPromise<TValueOrError<FString, FString>>(
+						MakeValue(ToolResult)).GetFuture();
+				}
+				else
+				{
+					return MakeFulfilledPromise<TValueOrError<FString, FString>>(
+						MakeError(ToolError)).GetFuture();
+				}
+			}
+			else
+			{
+				return MakeFulfilledPromise<TValueOrError<FString, FString>>(
+					MakeError(FString(TEXT("Invalid Toolset Handler")))).GetFuture();
+			}
+		}
+
+		// Returns the JSON schema for the toolset
+		virtual FString GetToolsetJsonSchema() const override
+		{
+			return ToolsetSchema;
+		}
+
+		// Returns the name of the toolset
+		virtual FString GetToolsetName() const override
+		{
+			return ToolsetName;
+		}
+
+		// Returns the version of the toolset
+		virtual FString GetToolsetVersion() const override
+		{
+			return ToolsetVersion;
+		}
+
+	  private:
+		TWeakObjectPtr<UAIAssistantToolset> ToolsetHandler;
+
+		// Caching the toolset name is required in order to be able to Unregister from the registry when the
+		// owning UAIAssistantToolset is garbage collected, because the registry uses the toolset name as keys
+		// into its map of handlers.  Cannot ask the owning UAIAssistantToolset for its name at that point, as
+		// it's being destroyed.
+		FString ToolsetName;
+
+		// Cached version, since caching the name was already necessary, might as well cache this as well.
+		FString ToolsetVersion;
+
+		// Cached schema, since caching the name was already necessary, might as well cache this as well.
+		FString ToolsetSchema;
+	};
+}
+
+void UAIAssistantToolset::PostInitProperties()
+{
+	Super::PostInitProperties();
+	// Only actual instances should register.
+	if (!HasAnyFlags(RF_ClassDefaultObject | RF_ArchetypeObject))
+	{
+		Register();
+	}
+}
+
+void UAIAssistantToolset::BeginDestroy()
+{
+	// Only actual instances should unregister.
+	if (!HasAnyFlags(RF_ClassDefaultObject | RF_ArchetypeObject))
+	{
+		Unregister();
+	}
+	Super::BeginDestroy();
+}
+
+FString UAIAssistantToolset::GetToolsetName() const
+{
+	return GetClass()->GetName();
+}
+
+FString UAIAssistantToolset::GetToolsetJsonSchema() const
+{
+	FString ToolsetName = GetToolsetName();
+	FString OutputString;
+	TSharedPtr<FJsonObject> Schema;
+
+	Schema = MakeShared<FJsonObject>();
+	check(Schema);
+
+	Schema->SetStringField(FString(TEXT("name")), ToolsetName);
+	Schema->SetStringField(FString(TEXT("version")), GetToolsetVersion());
+	Schema->SetStringField(FString(TEXT("description")), GetToolsetDescription());
+
+	TArray<TSharedPtr<FJsonValue>> ToolsArray;
+	for (TFieldIterator<UFunction> FuncIt(GetClass(), EFieldIterationFlags::None); FuncIt; ++FuncIt)
+	{
+		const UFunction* Function = *FuncIt;
+        if (Function)
+		{
+PRAGMA_DISABLE_DEPRECATION_WARNINGS
+
+			TSharedPtr<FJsonObject> ToolEntry = FJsonSchemaGenerator::UStructToJsonSchemaObject(Function);
+
+PRAGMA_ENABLE_DEPRECATION_WARNINGS
+			ToolsArray.Add(MakeShared<FJsonValueObject>(ToolEntry));
+		}
+	}
+	Schema->SetArrayField(FString(TEXT("tools")), ToolsArray);
+
+	if (!FJsonSerializer::Serialize(
+			Schema.ToSharedRef(),
+			TJsonWriterFactory<TCHAR, TCondensedJsonPrintPolicy<TCHAR>>::Create(&OutputString)))
+	{
+		UE_LOG(
+			LogAIAssistant, Warning,
+			TEXT("Failed to generate valid toolset schema for toolset %s"), *ToolsetName);
+	}
+
+	return OutputString;
+}
+
+FString UAIAssistantToolset::GetToolsetVersion() const
+{
+	FString Version = GetStringPropertyByName(TEXT("ToolsetVersion"));
+	if (Version.IsEmpty())
+	{
+		Version = FString(TEXT("1.0"));
+	}
+	return Version;
+}
+
+FString UAIAssistantToolset::GetToolsetDescription() const
+{
+	FString Description = GetStringPropertyByName(TEXT("ToolsetDescription"));
+	return Description;
+}
+
+void UAIAssistantToolset::Register()
+{
+	auto AssistantSubsystem = UAIAssistantSubsystem::Get(TEXT("Register"));
+	if (AssistantSubsystem.HasValue())
+	{
+		FString ToolsetName = GetToolsetName();
+		if (AssistantSubsystem.GetValue()->ToolsetRegistry.Find(ToolsetName))
+		{
+			UE_LOG(LogAIAssistant, Warning,
+				   TEXT("Toolset %s is already registered"), *ToolsetName);
+			return;
+		}
+
+		if (!RegistryHandler.IsValid())
+		{
+			RegistryHandler = MakeShared<UE::AIAssistant::FToolsetHandlerTrampoline>(
+				TWeakObjectPtr<UAIAssistantToolset>(this));
+		}
+
+		UE_LOG(LogAIAssistant, Display, TEXT("Registering Toolset %s"), *ToolsetName);
+		AssistantSubsystem.GetValue()->ToolsetRegistry.RegisterToolset(RegistryHandler);
+	}
+}
+
+void UAIAssistantToolset::Unregister()
+{
+	auto AssistantSubsystem = UAIAssistantSubsystem::Get(TEXT("Unregister"));
+	if (AssistantSubsystem.HasValue())
+	{
+		FString ToolsetName = GetToolsetName();
+		TSharedPtr<UE::AIAssistant::IToolsetHandler>* PtrToRegisteredHandler =
+			AssistantSubsystem.GetValue()->ToolsetRegistry.Find(ToolsetName);
+		TSharedPtr<UE::AIAssistant::IToolsetHandler> RegisteredHandler =
+			(PtrToRegisteredHandler) ? *PtrToRegisteredHandler : nullptr;
+		if (!RegisteredHandler)
+		{
+			// Nothing is registered, nothing to do.
+			return;
+		}
+
+		if (RegistryHandler.IsValid() && RegistryHandler == RegisteredHandler)
+		{
+			UE_LOG(LogAIAssistant, Display, TEXT("Unregistering Toolset %s"), *ToolsetName);
+			AssistantSubsystem.GetValue()->ToolsetRegistry.UnregisterToolset(RegistryHandler);
+		}
+	}
+}
+
+void UAIAssistantToolset::InvokeTool(
+	const FString& ToolName, const FString& JsonInput, FString& OutResult, FString& OutError)
+{
+	OutResult.Empty();
+	OutError.Empty();
+
+	if (RegisteredToolNames.IsEmpty())
+	{
+		RegisterToolNames();
+	}
+
+	UFunction* Function = GetClass()->FindFunctionByName(*ToolName);
+	if (!Function || !RegisteredToolNames.Contains(ToolName))
+	{
+		OutError = FString::Printf(TEXT("Unknown tool %s"), *ToolName);
+		return;
+	}
+
+	TSharedPtr<FJsonObject> ToolCallReturnValue;
+	UE::AIAssistant::FObjectFunctionToolCall::FFunctionInputParamsJson JsonParams;
+	JsonParams.Set<FString>(JsonInput);
+	const UE::AIAssistant::FObjectFunctionToolCall ToolCall =
+		UE::AIAssistant::FObjectFunctionToolCall(TNotNull<UObject*>(this), TNotNull<UFunction*>(Function));
+
+	UE::AIAssistant::FObjectFunctionToolCall::FToolCallResult ToolCallResult = ToolCall.Execute(JsonParams);
+
+	if (ToolCallResult.HasError())
+	{
+		FString Error = ToolCallResult.StealError().ToString();
+		OutError = FString::Printf(TEXT("Error executing tool %s: %s"), *ToolName, *Error);
+		return;
+	}
+
+	if (!ToolCallResult.HasValue())
+	{
+		OutError = FString::Printf(TEXT("No return value from executing tool %s"), *ToolName);
+		return;
+	}
+
+	ToolCallReturnValue = ToolCallResult.StealValue();
+	if (!ToolCallReturnValue.IsValid())
+	{
+		OutError = FString::Printf(TEXT("No return value from executing tool %s"), *ToolName);
+		return;
+	}
+
+	if (!FJsonSerializer::Serialize(
+			ToolCallReturnValue.ToSharedRef(),
+			TJsonWriterFactory<TCHAR, TCondensedJsonPrintPolicy<TCHAR>>::Create(&OutResult)))
+	{
+		OutError = FString::Printf(TEXT("Unable to serialize result from executing tool %s"), *ToolName);
+		return;
+	}
+
+	return;
+}
+
+FString UAIAssistantToolset::GetStringPropertyByName(const FName& PropName) const
+{
+	FString PropertyValue;
+	FProperty* Property = GetClass()->FindPropertyByName(PropName);
+	if (Property)
+	{
+		if (FStrProperty* PropertyAsString = CastField<FStrProperty>(Property))
+		{
+			const uint8* PropertyAddr = Property->ContainerPtrToValuePtr<uint8>(this);
+			PropertyValue = PropertyAsString->GetPropertyValue(PropertyAddr);
+		}
+	}
+
+	return PropertyValue;
+}
+
+void UAIAssistantToolset::RegisterToolNames()
+{
+	for (TFieldIterator<UFunction> FuncIt(GetClass(), EFieldIterationFlags::None); FuncIt; ++FuncIt)
+	{
+		const UFunction* Function = *FuncIt;
+        if (Function)
+		{
+			RegisteredToolNames.Add(Function->GetName());
+		}
+	}
+}
diff --git a/Engine/Plugins/Experimental/AIAssistant/Source/AIAssistant/Private/AIAssistantToolset.h b/Engine/Plugins/Experimental/AIAssistant/Source/AIAssistant/Private/AIAssistantToolset.h
new file mode 100644
--- /dev/null
+++ b/Engine/Plugins/Experimental/AIAssistant/Source/AIAssistant/Private/AIAssistantToolset.h
@@ -0,0 +1,92 @@
+// Copyright Epic Games, Inc. All Rights Reserved.
+
+#pragma once
+
+#include "CoreMinimal.h"
+
+#include "Containers/UnrealString.h"
+
+#include "AIAssistantToolsetHandler.h"
+
+#include "AIAssistantToolset.generated.h"
+
+namespace UE::AIAssistant
+{
+	class FToolsetHandlerTrampoline;
+}
+
+/**
+ * Provide AI Assistant toolset handler implementation.
+ *
+ */
+UCLASS(Blueprintable)
+class AIASSISTANT_API UAIAssistantToolset: public UObject
+{
+	GENERATED_BODY()
+
+public:
+
+	static inline FString DEFAULT_TOOLSET_VERSION = FString("1.0");
+
+	// Returns the JSON schema for the toolset
+	UFUNCTION(BlueprintCallable, Category=AssistantToolset)
+	FString GetToolsetJsonSchema() const;
+
+	// Returns the name of the toolset
+	UFUNCTION(BlueprintCallable, Category=AssistantToolset)
+	FString GetToolsetName() const;
+
+	// Returns the version of the toolset
+	UFUNCTION(BlueprintCallable, Category=AssistantToolset)
+	FString GetToolsetVersion() const;
+
+	// Returns the description of the toolset
+	UFUNCTION(BlueprintCallable, Category=AssistantToolset)
+	FString GetToolsetDescription() const;
+
+	// Registers the toolset with the AI Assistant.
+	UFUNCTION(BlueprintCallable, Category=AssistantToolset)
+	void Register();
+
+	// Unregisters the toolset from the AI Assistant.
+	UFUNCTION(BlueprintCallable, Category=AssistantToolset)
+	void Unregister();
+
+	// Provides toolset version to the AI Assistant.
+	// Override this in PostInitProperties.
+	UPROPERTY(BlueprintReadOnly, Category=AssistantToolset)
+	FString ToolsetVersion = DEFAULT_TOOLSET_VERSION;
+
+	// Provides toolset description to the AI Assistant.
+	// Override this in PostInitProperties.
+	UPROPERTY(BlueprintReadOnly, Category=AssistantToolset)
+	FString ToolsetDescription;
+
+	// Invoke a tool in this toolset.
+	// Successful call puts JSON result in OutResult.
+	// Puts error message in OutError if there was error in invocation.
+	UFUNCTION(BlueprintCallable, Category=AssistantToolset)
+	void InvokeTool(const FString& ToolName, const FString& JsonInput, FString& OutResult, FString& OutError);
+
+	// Overridden to automatically register toolset with AI Assistant subsystem.
+	// Subclasses should also override the ToolsetVersion, ToolsetDescription properties here.
+	virtual void PostInitProperties() override;
+	// Overridden to automatically unregister toolset with AI Assistant subsystem.
+	virtual void BeginDestroy() override;
+
+protected:
+	// Return the value of a string property given the property's name.
+	UFUNCTION()
+	FString GetStringPropertyByName(const FName& PropName) const;
+
+private:
+	// Generate list of valid tool names.
+	void RegisterToolNames();
+
+	// Handler that implements IToolsetHandler to register with the AI Assistant.
+	TSharedPtr<UE::AIAssistant::FToolsetHandlerTrampoline> RegistryHandler;
+
+	// Array of valid tool names present in toolset schema.
+	TArray<FString> RegisteredToolNames;
+};
+
diff --git a/Engine/Plugins/Experimental/AIAssistant/Source/AIAssistant/Private/Tests/AIAssistantToolsetTest.cpp b/Engine/Plugins/Experimental/AIAssistant/Source/AIAssistant/Private/Tests/AIAssistantToolsetTest.cpp
new file mode 100644
--- /dev/null
+++ b/Engine/Plugins/Experimental/AIAssistant/Source/AIAssistant/Private/Tests/AIAssistantToolsetTest.cpp
@@ -0,0 +1,342 @@
+// Copyright Epic Games, Inc. All Rights Reserved.
+
+#include "AIAssistantToolsetTest.h"
+
+#include "CoreMinimal.h"
+
+#include "Misc/AutomationTest.h"
+#include "Dom/JsonObject.h"
+#include "JsonDomBuilder.h"
+#include "IPythonScriptPlugin.h"
+#include "PythonScriptTypes.h"
+#include "Policies/CondensedJsonPrintPolicy.h"
+#include "Serialization/JsonWriter.h"
+#include "Serialization/JsonSerializer.h"
+
+#include "Tests/AIAssistantTestFlags.h"
+#include "AIAssistantBPToolsetRegistry.h"
+#include "AIAssistantLog.h"
+
+
+#if WITH_DEV_AUTOMATION_TESTS
+
+const FString FakeToolsetName(TEXT("FakeToolset"));
+const FString FakeToolsetInput(TEXT(R"({"inString":"hello"})"));
+const FString FakeToolName(TEXT("SimpleTool"));
+
+// Wrapper for creating and cleaning up UObjects.
+template<typename T>class ToolsetContainer
+{
+public:
+	ToolsetContainer()
+	{
+		Instance = NewObject<T>();
+	}
+
+	~ToolsetContainer()
+	{
+		Instance->MarkAsGarbage();
+		Instance = nullptr;
+		CollectGarbage(GARBAGE_COLLECTION_KEEPFLAGS);
+	}
+
+	TObjectPtr<T> operator*() const
+	{
+		return Instance;
+	}
+
+private:
+	TObjectPtr<T> Instance;
+};
+
+// Test that toolset class creation registers toolset and destroy unregisters toolset.
+IMPLEMENT_SIMPLE_AUTOMATION_TEST(
+	FAIAssistantToolsetRegistersOnCreateAndDestroy,
+	"AI.Assistant.Toolset.RegistersOnCreateAndDestroy",
+	AIAssistantTest::Flags)
+
+bool FAIAssistantToolsetRegistersOnCreateAndDestroy::RunTest(const FString& UnusedParameters)
+{
+	// Ensure clean registry.
+	CollectGarbage(GARBAGE_COLLECTION_KEEPFLAGS);
+	UAIAssistantToolsetRegistry::UnregisterToolsetByName(FakeToolsetName);
+
+	TestFalse("No toolsets registered to start.",
+			  UAIAssistantToolsetRegistry::IsToolsetRegistered(FakeToolsetName));
+	{
+		ToolsetContainer<UFakeToolset> Container;
+		TestTrue("Toolset registered after class instance created.",
+				 UAIAssistantToolsetRegistry::IsToolsetRegistered(FakeToolsetName));
+	}
+	TestFalse("Toolset is unregistered after class instance destroyed.",
+			  UAIAssistantToolsetRegistry::IsToolsetRegistered(FakeToolsetName));
+
+	return true;
+}
+
+// Test that toolset can return its name, version, and description.
+IMPLEMENT_SIMPLE_AUTOMATION_TEST(
+	FAIAssistantToolsetCanGetValues,
+	"AI.Assistant.Toolset.CanGetValues",
+	AIAssistantTest::Flags)
+
+bool FAIAssistantToolsetCanGetValues::RunTest(const FString& UnusedParameters)
+{
+	CollectGarbage(GARBAGE_COLLECTION_KEEPFLAGS);
+
+	ToolsetContainer<UFakeToolset> Container;
+	TObjectPtr<UFakeToolset> Instance = *Container;
+	TestEqual("Toolset returns its own name",
+			  Instance->GetToolsetName(), FakeToolsetName);
+	TestEqual("Toolset returns default non-overridden version",
+			  Instance->GetToolsetVersion(), UAIAssistantToolset::DEFAULT_TOOLSET_VERSION);
+	TestEqual("Toolset returns default non-overridden description",
+			  Instance->GetToolsetDescription(), FString());
+	return true;
+}
+
+// Test that overriding toolset properties for version, description works properly.
+IMPLEMENT_SIMPLE_AUTOMATION_TEST(
+	FAIAssistantToolsetCanOverrideProperties,
+	"AI.Assistant.Toolset.CanOverrideProperties",
+	AIAssistantTest::Flags)
+
+bool FAIAssistantToolsetCanOverrideProperties::RunTest(const FString& UnusedParameters)
+{
+	CollectGarbage(GARBAGE_COLLECTION_KEEPFLAGS);
+
+	ToolsetContainer<UFakeToolset2> Container;
+	TObjectPtr<UFakeToolset2> Instance = *Container;
+	TestEqual("Toolset returns overridden version",
+			  Instance->GetToolsetVersion(), UFakeToolset2::MY_VERSION);
+	TestEqual("Toolset returns overridden description",
+			  Instance->GetToolsetDescription(), UFakeToolset2::MY_DESCRIPTION);
+	return true;
+}
+
+TSharedPtr<FJsonObject> GetFakeToolsetSchema(const TObjectPtr<UFakeToolset> Toolset)
+{
+	FJsonDomBuilder::FObject SchemaBuilder;
+
+	SchemaBuilder.Set(TEXT("name"), Toolset->GetToolsetName());
+	SchemaBuilder.Set(TEXT("version"), Toolset->GetToolsetVersion());
+	SchemaBuilder.Set(TEXT("description"), Toolset->GetToolsetDescription());
+
+	FJsonDomBuilder::FArray ToolsArray;
+	{
+		FJsonDomBuilder::FObject ToolEntry;
+		ToolEntry.Set(TEXT("type"), TEXT("object"));
+		ToolEntry.Set(TEXT("description"), TEXT("Simple tool that does nothing."));
+		ToolEntry.Set(TEXT("name"), TEXT("FakeToolset.SimpleTool"));
+		{
+			FJsonDomBuilder::FObject InputSchema;
+			InputSchema.Set(TEXT("type"), TEXT("object"));
+			FJsonDomBuilder::FObject SchemaProps;
+			FJsonDomBuilder::FObject SchemaPropsType;
+			SchemaPropsType.Set(TEXT("type"), TEXT("string"));
+			SchemaProps.Set(TEXT("inString"), SchemaPropsType);
+			InputSchema.Set(TEXT("properties"), SchemaProps);
+			FJsonDomBuilder::FArray RequiredArray;
+			RequiredArray.Add(TEXT("inString"));
+			InputSchema.Set(TEXT("required"), RequiredArray);
+			ToolEntry.Set(TEXT("inputSchema"), InputSchema);
+		}
+		{
+			FJsonDomBuilder::FObject OutputSchema;
+			OutputSchema.Set(TEXT("type"), TEXT("object"));
+			FJsonDomBuilder::FObject SchemaProps;
+			FJsonDomBuilder::FObject SchemaPropsType;
+			SchemaPropsType.Set(TEXT("type"), TEXT("string"));
+			SchemaProps.Set(TEXT("returnValue"), SchemaPropsType);
+			OutputSchema.Set(TEXT("properties"), SchemaProps);
+			FJsonDomBuilder::FArray RequiredArray;
+			RequiredArray.Add(TEXT("returnValue"));
+			OutputSchema.Set(TEXT("required"), RequiredArray);
+			ToolEntry.Set(TEXT("outputSchema"), OutputSchema);
+		}
+		ToolsArray.Add(ToolEntry);
+	}
+	SchemaBuilder.Set(TEXT("tools"), ToolsArray);
+	return SchemaBuilder.AsJsonObject();
+}
+
+// Test that toolset can return a schema properly.
+IMPLEMENT_SIMPLE_AUTOMATION_TEST(
+	FAIAssistantToolsetGetSchema,
+	"AI.Assistant.Toolset.GetSchema",
+	AIAssistantTest::Flags)
+
+bool FAIAssistantToolsetGetSchema::RunTest(const FString& UnusedParameters)
+{
+	CollectGarbage(GARBAGE_COLLECTION_KEEPFLAGS);
+	ToolsetContainer<UFakeToolset> Container;
+	TObjectPtr<UFakeToolset> Instance = *Container;
+
+	TSharedPtr<FJsonObject> ExpectedSchema = GetFakeToolsetSchema(Instance);
+
+	FString ExpectedSchemaJson;
+	FJsonSerializer::Serialize(
+		ExpectedSchema.ToSharedRef(),
+		TJsonWriterFactory<TCHAR, TCondensedJsonPrintPolicy<TCHAR>>::Create(&ExpectedSchemaJson));
+	FString OutputSchemaJson = Instance->GetToolsetJsonSchema();
+	TestEqual("Toolset schema matches expected", OutputSchemaJson, ExpectedSchemaJson);
+	return true;
+}
+
+// Test that toolset can unregister and re-register itself properly.
+IMPLEMENT_SIMPLE_AUTOMATION_TEST(
+	FAIAssistantToolsetRegisterUnregisterItself,
+	"AI.Assistant.Toolset.RegisterUnregisterItself",
+	AIAssistantTest::Flags)
+
+bool FAIAssistantToolsetRegisterUnregisterItself::RunTest(const FString& UnusedParameters)
+{
+	CollectGarbage(GARBAGE_COLLECTION_KEEPFLAGS);
+	ToolsetContainer<UFakeToolset> Container;
+	TObjectPtr<UFakeToolset> Instance = *Container;
+
+	TestTrue("Toolset registered after class instance created.",
+			  UAIAssistantToolsetRegistry::IsToolsetRegistered(FakeToolsetName));
+	Instance->Unregister();
+	TestFalse("Toolset unregistered after unregistering itself.",
+			  UAIAssistantToolsetRegistry::IsToolsetRegistered(FakeToolsetName));
+	Instance->Register();
+	TestTrue("Toolset registered after re-registering itself.",
+			  UAIAssistantToolsetRegistry::IsToolsetRegistered(FakeToolsetName));
+	return true;
+}
+
+// Test that invoking a tool from a toolset calls the right method and returns the expected result.
+IMPLEMENT_SIMPLE_AUTOMATION_TEST(
+	FAIAssistantToolsetExecuteTool,
+	"AI.Assistant.Toolset.ExecuteTool",
+	AIAssistantTest::Flags)
+
+bool FAIAssistantToolsetExecuteTool::RunTest(const FString& UnusedParameters)
+{
+	CollectGarbage(GARBAGE_COLLECTION_KEEPFLAGS);
+	ToolsetContainer<UFakeToolset> Container;
+
+	FString OutResult;
+	FString OutError;
+	FString InputString = FakeToolsetInput;
+
+	UAIAssistantToolsetRegistry::ExecuteTool(
+		FakeToolsetName, FakeToolName, InputString, OutResult, OutError);
+	TestTrue("Tool execution returns no error", OutError.IsEmpty());
+
+	FString ExpectedOutput = FString(TEXT(R"({"returnValue":"hello"})"));
+	TestEqual("Tool execution returns expected output", OutResult, ExpectedOutput);
+	return true;
+}
+
+// Test that invoking a tool from an invalid toolset errors out properly.
+IMPLEMENT_SIMPLE_AUTOMATION_TEST(
+	FAIAssistantToolsetExecuteToolInvalidToolset,
+	"AI.Assistant.Toolset.ExecuteToolInvalidToolset",
+	AIAssistantTest::Flags)
+
+bool FAIAssistantToolsetExecuteToolInvalidToolset::RunTest(const FString& UnusedParameters)
+{
+	CollectGarbage(GARBAGE_COLLECTION_KEEPFLAGS);
+	ToolsetContainer<UFakeToolset> Container;
+
+	FString OutResult;
+	FString OutError;
+	FString InputString = FakeToolsetInput;
+
+	AddExpectedError(TEXT("not found"), EAutomationExpectedErrorFlags::Contains);
+	UAIAssistantToolsetRegistry::ExecuteTool(
+		FString("FakeToolsetInvalid"), FakeToolName, InputString, OutResult, OutError);
+	TestFalse("Tool execution with invalid toolset name returns error", OutError.IsEmpty());
+	return true;
+}
+
+// Test that invoking a tool with an invalid tool name errors out properly.
+IMPLEMENT_SIMPLE_AUTOMATION_TEST(
+	FAIAssistantToolsetExecuteToolInvalidToolname,
+	"AI.Assistant.Toolset.ExecuteToolInvalidToolname",
+	AIAssistantTest::Flags)
+
+bool FAIAssistantToolsetExecuteToolInvalidToolname::RunTest(const FString& UnusedParameters)
+{
+	CollectGarbage(GARBAGE_COLLECTION_KEEPFLAGS);
+	ToolsetContainer<UFakeToolset> Container;
+
+	FString OutResult;
+	FString OutError;
+	FString InputString = FakeToolsetInput;
+
+	UAIAssistantToolsetRegistry::ExecuteTool(
+		FakeToolsetName, FString("SimpleToolInvalid"), InputString, OutResult, OutError);
+	TestFalse("Tool execution with invalid tool name returns error", OutError.IsEmpty());
+	return true;
+}
+
+// Test that invoking a tool with an existing but invalid tool name errors out properly.
+IMPLEMENT_SIMPLE_AUTOMATION_TEST(
+	FAIAssistantToolsetExecuteToolExistingButInvalidToolname,
+	"AI.Assistant.Toolset.ExecuteToolExistingButInvalidToolname",
+	AIAssistantTest::Flags)
+
+bool FAIAssistantToolsetExecuteToolExistingButInvalidToolname::RunTest(const FString& UnusedParameters)
+{
+	CollectGarbage(GARBAGE_COLLECTION_KEEPFLAGS);
+	ToolsetContainer<UFakeToolset> Container;
+
+	FString OutResult;
+	FString OutError;
+	FString InputString = FakeToolsetInput;
+
+	UAIAssistantToolsetRegistry::ExecuteTool(
+		FakeToolsetName, FString("Register"), InputString, OutResult, OutError);
+	TestFalse("Tool execution with invalid tool name returns error", OutError.IsEmpty());
+	return true;
+}
+
+// Test that invoking a tool with a valid but unregistered toolset errors out properly.
+IMPLEMENT_SIMPLE_AUTOMATION_TEST(
+	FAIAssistantToolsetExecuteToolUnregisteredToolset,
+	"AI.Assistant.Toolset.ExecuteToolUnregisteredToolset",
+	AIAssistantTest::Flags)
+
+bool FAIAssistantToolsetExecuteToolUnregisteredToolset::RunTest(const FString& UnusedParameters)
+{
+	CollectGarbage(GARBAGE_COLLECTION_KEEPFLAGS);
+	ToolsetContainer<UFakeToolset> Container;
+
+	FString OutResult;
+	FString OutError;
+	FString InputString = FakeToolsetInput;
+
+	TObjectPtr<UFakeToolset> Instance = *Container;
+	Instance->Unregister();
+
+	AddExpectedError(TEXT("not found"), EAutomationExpectedErrorFlags::Contains);
+	UAIAssistantToolsetRegistry::ExecuteTool(
+		FakeToolsetName, FakeToolName, InputString, OutResult, OutError);
+	TestFalse("Tool execution with unregistered toolset returns error", OutError.IsEmpty());
+	return true;
+}
+
+// Run python tests
+IMPLEMENT_SIMPLE_AUTOMATION_TEST(
+	FAIAssistantToolsetPythonTests,
+	"AI.Assistant.Toolset.PythonTests",
+	AIAssistantTest::Flags)
+
+bool FAIAssistantToolsetPythonTests::RunTest(const FString& UnusedParameters)
+{
+	if (!IPythonScriptPlugin::Get()->IsPythonAvailable())
+	{
+		return true;
+	}
+
+	FPythonCommandEx PythonCommand = FPythonCommandEx();
+	PythonCommand.ExecutionMode = EPythonCommandExecutionMode::ExecuteStatement;
+	PythonCommand.Command = TEXT("import ai_assistant.tests; ai_assistant.tests.run_tests()");
+	return IPythonScriptPlugin::Get()->ExecPythonCommandEx(PythonCommand);
+}
+
+#endif	// WITH_DEV_AUTOMATION_TESTS
+
diff --git a/Engine/Plugins/Experimental/AIAssistant/Source/AIAssistant/Private/Tests/AIAssistantToolsetTest.h b/Engine/Plugins/Experimental/AIAssistant/Source/AIAssistant/Private/Tests/AIAssistantToolsetTest.h
new file mode 100644
--- /dev/null
+++ b/Engine/Plugins/Experimental/AIAssistant/Source/AIAssistant/Private/Tests/AIAssistantToolsetTest.h
@@ -0,0 +1,53 @@
+// Copyright Epic Games, Inc. All Rights Reserved.
+
+#pragma once
+
+#include "CoreMinimal.h"
+
+#include "AIAssistantToolset.h"
+
+#include "AIAssistantToolsetTest.generated.h"
+
+/** Fake Toolset that does nothing. */
+UCLASS(Blueprintable)
+class UFakeToolset : public  UAIAssistantToolset
+{
+	GENERATED_BODY()
+
+public:
+
+	/** Simple tool that does nothing. */
+	UFUNCTION(BlueprintCallable, Category=Fake)
+	FString SimpleTool(const FString& InString)
+	{
+		return InString;
+	}
+};
+
+/** Fake Toolset that does nothing except update version and description. */
+UCLASS(Blueprintable)
+class UFakeToolset2 : public  UAIAssistantToolset
+{
+	GENERATED_BODY()
+
+public:
+
+	static inline FString MY_VERSION = FString("1.999.the_best");
+	static inline FString MY_DESCRIPTION = FString("Do nothing number two");
+
+	/** Simple tool that does nothing. */
+	UFUNCTION(BlueprintCallable, Category=Fake)
+	FString SimpleTool(const FString& InString)
+	{
+		return InString;
+	}
+
+	// Override the ToolsetVersion, ToolsetDescription properties.
+	virtual void PostInitProperties() override
+	{
+		ToolsetVersion = MY_VERSION;
+		ToolsetDescription = MY_DESCRIPTION;
+		Super::PostInitProperties();
+	}
+
+};
