[AIAssistant] Fix WebApplication tear down and a hang when interacting with the assistant.

#Synced-CL 48894789

diff --git a/Engine/Plugins/Experimental/AIAssistant/Source/AIAssistant/Private/AIAssistantWebApplication.cpp b/Engine/Plugins/Experimental/AIAssistant/Source/AIAssistant/Private/AIAssistantWebApplication.cpp
--- a/Engine/Plugins/Experimental/AIAssistant/Source/AIAssistant/Private/AIAssistantWebApplication.cpp
+++ b/Engine/Plugins/Experimental/AIAssistant/Source/AIAssistant/Private/AIAssistantWebApplication.cpp
@@ -14,6 +14,8 @@
 
 namespace UE::AIAssistant
 {
+	static const TCHAR* WebApplicationShutDownMessage = TEXT("WebApplication Shut Down");
+
 	FWebApplication::FWebApplication(
 		TFunction<TSharedPtr<FWebApi>()>&& WebApiFactory) :
 		ConversationReadyExecutor(
@@ -37,7 +39,6 @@ namespace UE::AIAssistant
 		{
 			FInternationalization::Get().OnCultureChanged().Remove(CultureChangeDelegateHandle);
 		}
-		ResetConversationState();
 	}
 
 	void FWebApplication::CreateConversation()
@@ -185,21 +186,26 @@ namespace UE::AIAssistant
 		return AIAssistantSubsystem->ToolsetRegistry;
 	}
 
-	TFuture<bool> FWebApplication::OnConversationUpdate(const FConversationId& ConversationId)
+	TFuture<TValueOrError<void, FString>> FWebApplication::OnConversationUpdate(
+		const FConversationId& ConversationId)
 	{
-		TFuture<bool> Success;
+		auto Promise = MakeShared<TPromise<TValueOrError<void, FString>>>();
 		if (!WithWebApi(
-			[this, ConversationId, &Success](FWebApi& WebApi) -> void
+			[this, ConversationId, Promise](FWebApi& WebApi) -> void
 			{
-				Success = WebApi.GetConversation(ConversationId).Next(
-					[WeakThis = SharedThis(this).ToWeakPtr()](auto Result) -> bool
+				WebApi.GetConversation(ConversationId).Next(
+					[WeakThis = SharedThis(this).ToWeakPtr(), Promise](const auto& Result) -> void
 					{
 						auto This = WeakThis.Pin();
-						if (!This) return false;
-
-						if (This->HandleErrorResult(TEXT("OnConversationUpdate"), Result))
+						if (!This)
+						{
+							Promise->SetValue(MakeError(WebApplicationShutDownMessage));
+							return;
+						}
+						if (Result.HasError())
 						{
-							return false;
+							Promise->SetValue(MakeError(Result.GetError()));
+							return;
 						}
 						// Process each message in the conversation.
 						const auto& Conversation = Result.GetValue();
@@ -207,14 +213,14 @@ namespace UE::AIAssistant
 						{
 							This->ProcessMessage(Message);
 						}
-						return true;
+						Promise->SetValue(MakeValue());
 					});
 			}))
 		{
 			// If the web API is not available, return a failed future.
-			Success = MakeFulfilledPromise<bool>(false).GetFuture();
+			Promise->SetValue(MakeError(WebApplicationShutDownMessage));
 		}
-		return Success;
+		return Promise->GetFuture();
 	}
 
 	void FWebApplication::ProcessMessage(const FMessage& Message)
@@ -244,7 +250,7 @@ namespace UE::AIAssistant
 				return This
 					? This->ProcessToolCallContent(ToolCall)
 					:  MakeFulfilledPromise<TValueOrError<void, FString>>(
-						MakeError("WebApplication Shut Down")).GetFuture();
+						MakeError(WebApplicationShutDownMessage)).GetFuture();
 			}).Next([Message](const TValueOrError<void, FString>& FinalResult)
 				{
 					if (FinalResult.HasError())
@@ -280,7 +286,7 @@ namespace UE::AIAssistant
 				const TValueOrError<FString, FString>& ToolResult) -> TValueOrError<void, FString>
 			{
 				auto This = WeakThis.Pin();
-				if (!This) return MakeError("WebApplication Shut Down");
+				if (!This) return MakeError(WebApplicationShutDownMessage);
 
 				if (ResponseRequired)
 				{
@@ -324,25 +330,32 @@ namespace UE::AIAssistant
 			[WeakThis = SharedThis(this).ToWeakPtr()](FWebApi& WebApi)
 			{
 				auto This = WeakThis.Pin();
-				if (!This) return false;
+				if (!This) return;
 
-				auto HandleOrError = WebApi.RegisterOnConversationUpdate(
+				WebApi.RegisterOnConversationUpdate(
 					[WeakThis](FConversationId Id) -> bool
 					{
 						auto This = WeakThis.Pin();
 						if (!This) return false;
-						// This returns a future, but in practice the web API implementation will
-						// complete immediately so we can block here.
-						return This->OnConversationUpdate(Id).Get();
-					}).Get();
 
-				if (This->HandleErrorResult(
-					TEXT("Error registering OnConversationUpdate"), HandleOrError))
-				{
-					return false;
-				}
-				This->OnConversationUpdateHandle.Emplace(HandleOrError.GetValue());
-				return true;
+						This->OnConversationUpdate(Id).Next(
+							[WeakThis](const auto& Result) {
+								auto This = WeakThis.Pin();
+								return This && !This->HandleErrorResult(
+									TEXT("OnConversationUpdate"), Result);
+							});
+						return true;
+					}).Next([WeakThis](auto HandleOrError) {
+						auto This = WeakThis.Pin();
+						if (!This) return;
+
+						if (This->HandleErrorResult(
+							TEXT("Error registering OnConversationUpdate"), HandleOrError))
+						{
+							return;
+						}
+						This->OnConversationUpdateHandle.Emplace(HandleOrError.GetValue());
+					});
 			});
 	}
 
diff --git a/Engine/Plugins/Experimental/AIAssistant/Source/AIAssistant/Private/AIAssistantWebApplication.h b/Engine/Plugins/Experimental/AIAssistant/Source/AIAssistant/Private/AIAssistantWebApplication.h
--- a/Engine/Plugins/Experimental/AIAssistant/Source/AIAssistant/Private/AIAssistantWebApplication.h
+++ b/Engine/Plugins/Experimental/AIAssistant/Source/AIAssistant/Private/AIAssistantWebApplication.h
@@ -131,7 +131,7 @@ namespace UE::AIAssistant
 		virtual FToolsetRegistry& GetToolsetRegistry();
 
 		// Handle updates to the conversation.
-		TFuture<bool> OnConversationUpdate(const FConversationId& ConversationId);
+		TFuture<TValueOrError<void, FString>> OnConversationUpdate(const FConversationId& ConversationId);
 
 		// Process a message from the conversation.
 		virtual void ProcessMessage(const FMessage& Message);
diff --git a/Engine/Plugins/Experimental/AIAssistant/Source/AIAssistant/Private/Tests/AIAssistantWebApplicationTest.cpp b/Engine/Plugins/Experimental/AIAssistant/Source/AIAssistant/Private/Tests/AIAssistantWebApplicationTest.cpp
--- a/Engine/Plugins/Experimental/AIAssistant/Source/AIAssistant/Private/Tests/AIAssistantWebApplicationTest.cpp
+++ b/Engine/Plugins/Experimental/AIAssistant/Source/AIAssistant/Private/Tests/AIAssistantWebApplicationTest.cpp
@@ -745,7 +745,7 @@ bool FAIAssistantWebApplicationOnConversationUpdateProcessMessages::RunTest(
 		*this, nullptr, TEXT("getConversation"), *Id.ToJson(false), *Conversation.ToJson(false),
 		false);
 
-	if (TestTrue(TEXT("OnConversationUpdate_Success"), Future.Get())
+	if (TestTrue(TEXT("OnConversationUpdate_Success"), Future.Get().HasValue())
 		&& TestEqual(
 			TEXT("OnConversationUpdate_CallCount"),
 			WebApplication->ProcessMessageCalls.Num(), 2))
