[AI Assistant] Fix some issues with UAIAssistantToolset

- Shorten UAIAssistantToolset method names by removing redundant "Toolset" from name
- Simplify GetVersion()
- Remove ToolsetDescription property: GetDescription() will automatically get it from metadata (and toolsets defined in Python will use the toolset class docstring now)
- Schema: skip UFUNCTIONs marked as BlueprintInternalUseOnly
- check that setting metadata (e.g. ClampMin, ClampMax) on UProperties gets through to schema
- ensure that list of valid tool methods is consistent between generating schema and validating invocation calls

#rnx

#Synced-CL 49089041

diff --git a/Engine/Plugins/Experimental/AIAssistant/Source/AIAssistant/Private/AIAssistantToolset.cpp b/Engine/Plugins/Experimental/AIAssistant/Source/AIAssistant/Private/AIAssistantToolset.cpp
--- a/Engine/Plugins/Experimental/AIAssistant/Source/AIAssistant/Private/AIAssistantToolset.cpp
+++ b/Engine/Plugins/Experimental/AIAssistant/Source/AIAssistant/Private/AIAssistantToolset.cpp
@@ -36,9 +36,9 @@ namespace UE::AIAssistant
 			TStrongObjectPtr<UAIAssistantToolset> Handler = ToolsetHandler.Pin();
 			if (Handler)
 			{
-				ToolsetName = Handler->GetToolsetName();
-				ToolsetVersion = Handler->GetToolsetVersion();
-				ToolsetSchema = Handler->GetToolsetJsonSchema();
+				ToolsetName = Handler->GetName();
+				ToolsetVersion = Handler->GetVersion();
+				ToolsetSchema = Handler->GetJsonSchema();
 			}
 		}
 
@@ -126,14 +126,9 @@ void UAIAssistantToolset::BeginDestroy()
 	Super::BeginDestroy();
 }
 
-FString UAIAssistantToolset::GetToolsetName() const
+FString UAIAssistantToolset::GetJsonSchema() const
 {
-	return GetClass()->GetName();
-}
-
-FString UAIAssistantToolset::GetToolsetJsonSchema() const
-{
-	FString ToolsetName = GetToolsetName();
+	FString ToolsetName = GetName();
 	FString OutputString;
 	TSharedPtr<FJsonObject> Schema;
 
@@ -141,20 +136,20 @@ FString UAIAssistantToolset::GetToolsetJsonSchema() const
 	check(Schema);
 
 	Schema->SetStringField(FString(TEXT("name")), ToolsetName);
-	Schema->SetStringField(FString(TEXT("version")), GetToolsetVersion());
-	Schema->SetStringField(FString(TEXT("description")), GetToolsetDescription());
+	Schema->SetStringField(FString(TEXT("version")), GetVersion());
+	Schema->SetStringField(FString(TEXT("description")), GetDescription());
 
 	TArray<TSharedPtr<FJsonValue>> ToolsArray;
-	for (TFieldIterator<UFunction> FuncIt(GetClass(), EFieldIterationFlags::None); FuncIt; ++FuncIt)
+	TArray<TObjectPtr<UFunction>> ValidTools = GetValidToolMethods();
+	for (UFunction* Function : ValidTools)
 	{
-		const UFunction* Function = *FuncIt;
         if (Function)
 		{
-PRAGMA_DISABLE_DEPRECATION_WARNINGS
 
+PRAGMA_DISABLE_DEPRECATION_WARNINGS
 			TSharedPtr<FJsonObject> ToolEntry = FJsonSchemaGenerator::UStructToJsonSchemaObject(Function);
-
 PRAGMA_ENABLE_DEPRECATION_WARNINGS
+
 			ToolsArray.Add(MakeShared<FJsonValueObject>(ToolEntry));
 		}
 	}
@@ -172,28 +167,12 @@ PRAGMA_ENABLE_DEPRECATION_WARNINGS
 	return OutputString;
 }
 
-FString UAIAssistantToolset::GetToolsetVersion() const
-{
-	FString Version = GetStringPropertyByName(TEXT("ToolsetVersion"));
-	if (Version.IsEmpty())
-	{
-		Version = FString(TEXT("1.0"));
-	}
-	return Version;
-}
-
-FString UAIAssistantToolset::GetToolsetDescription() const
-{
-	FString Description = GetStringPropertyByName(TEXT("ToolsetDescription"));
-	return Description;
-}
-
 void UAIAssistantToolset::Register()
 {
 	auto AssistantSubsystem = UAIAssistantSubsystem::Get(TEXT("Register"));
 	if (AssistantSubsystem.HasValue())
 	{
-		FString ToolsetName = GetToolsetName();
+		FString ToolsetName = GetName();
 		if (AssistantSubsystem.GetValue()->ToolsetRegistry.Find(ToolsetName))
 		{
 			UE_LOG(LogAIAssistant, Warning,
@@ -217,7 +196,7 @@ void UAIAssistantToolset::Unregister()
 	auto AssistantSubsystem = UAIAssistantSubsystem::Get(TEXT("Unregister"));
 	if (AssistantSubsystem.HasValue())
 	{
-		FString ToolsetName = GetToolsetName();
+		FString ToolsetName = GetName();
 		TSharedPtr<UE::AIAssistant::IToolsetHandler>* PtrToRegisteredHandler =
 			AssistantSubsystem.GetValue()->ToolsetRegistry.Find(ToolsetName);
 		TSharedPtr<UE::AIAssistant::IToolsetHandler> RegisteredHandler =
@@ -293,30 +272,34 @@ void UAIAssistantToolset::InvokeTool(
 	return;
 }
 
-FString UAIAssistantToolset::GetStringPropertyByName(const FName& PropName) const
+void UAIAssistantToolset::RegisterToolNames()
 {
-	FString PropertyValue;
-	FProperty* Property = GetClass()->FindPropertyByName(PropName);
-	if (Property)
+	TArray<TObjectPtr<UFunction>> ValidTools = GetValidToolMethods();
+
+	for (UFunction* Function : ValidTools)
 	{
-		if (FStrProperty* PropertyAsString = CastField<FStrProperty>(Property))
+        if (Function)
 		{
-			const uint8* PropertyAddr = Property->ContainerPtrToValuePtr<uint8>(this);
-			PropertyValue = PropertyAsString->GetPropertyValue(PropertyAddr);
+			RegisteredToolNames.Add(Function->GetName());
 		}
 	}
-
-	return PropertyValue;
 }
 
-void UAIAssistantToolset::RegisterToolNames()
+TArray<TObjectPtr<UFunction>> UAIAssistantToolset::GetValidToolMethods() const
 {
+	TArray<TObjectPtr<UFunction>> ValidTools;
 	for (TFieldIterator<UFunction> FuncIt(GetClass(), EFieldIterationFlags::None); FuncIt; ++FuncIt)
 	{
-		const UFunction* Function = *FuncIt;
+		TObjectPtr<UFunction> Function = *FuncIt;
         if (Function)
 		{
-			RegisteredToolNames.Add(Function->GetName());
+			// Skip functions marked as for internal use only.
+			if (Function->HasMetaData(TEXT("BlueprintInternalUseOnly")))
+			{
+				continue;
+			}
+			ValidTools.Add(Function);
 		}
 	}
+	return ValidTools;
 }
diff --git a/Engine/Plugins/Experimental/AIAssistant/Source/AIAssistant/Private/AIAssistantToolset.h b/Engine/Plugins/Experimental/AIAssistant/Source/AIAssistant/Private/AIAssistantToolset.h
--- a/Engine/Plugins/Experimental/AIAssistant/Source/AIAssistant/Private/AIAssistantToolset.h
+++ b/Engine/Plugins/Experimental/AIAssistant/Source/AIAssistant/Private/AIAssistantToolset.h
@@ -30,19 +30,28 @@ public:
 
 	// Returns the JSON schema for the toolset
 	UFUNCTION(BlueprintCallable, Category=AssistantToolset)
-	FString GetToolsetJsonSchema() const;
+	FString GetJsonSchema() const;
 
 	// Returns the name of the toolset
 	UFUNCTION(BlueprintCallable, Category=AssistantToolset)
-	FString GetToolsetName() const;
+	FString GetName() const
+	{
+		return GetClass()->GetName();
+	}
 
 	// Returns the version of the toolset
 	UFUNCTION(BlueprintCallable, Category=AssistantToolset)
-	FString GetToolsetVersion() const;
+	FString GetVersion() const
+	{
+		return ToolsetVersion;
+	}
 
 	// Returns the description of the toolset
 	UFUNCTION(BlueprintCallable, Category=AssistantToolset)
-	FString GetToolsetDescription() const;
+	FString GetDescription() const
+	{
+		return GetClass()->GetToolTipText().ToString();
+	}
 
 	// Registers the toolset with the AI Assistant.
 	UFUNCTION(BlueprintCallable, Category=AssistantToolset)
@@ -57,28 +66,23 @@ public:
 	UPROPERTY(BlueprintReadOnly, Category=AssistantToolset)
 	FString ToolsetVersion = DEFAULT_TOOLSET_VERSION;
 
-	// Provides toolset description to the AI Assistant.
-	// Override this in PostInitProperties.
-	UPROPERTY(BlueprintReadOnly, Category=AssistantToolset)
-	FString ToolsetDescription;
-
-	// Invoke a tool in this toolset.
-	// Successful call puts JSON result in OutResult.
-	// Puts error message in OutError if there was error in invocation.
+	// Invoke a tool in this toolset by name.
+	//
+	// Expected to be called by the toolset registry ExecuteTool method, so
+	// ToolName should be the unqualified function name.
+	//
+	// Returns JSON result in OutResult unless there was an error.
+	//
+	// Returns error message in OutError if there was error in invocation.
 	UFUNCTION(BlueprintCallable, Category=AssistantToolset)
 	void InvokeTool(const FString& ToolName, const FString& JsonInput, FString& OutResult, FString& OutError);
 
 	// Overridden to automatically register toolset with AI Assistant subsystem.
-	// Subclasses should also override the ToolsetVersion, ToolsetDescription properties here.
+	// Subclasses should also override the ToolsetVersion property here.
 	virtual void PostInitProperties() override;
 	// Overridden to automatically unregister toolset with AI Assistant subsystem.
 	virtual void BeginDestroy() override;
 
-protected:
-	// Return the value of a string property given the property's name.
-	UFUNCTION()
-	FString GetStringPropertyByName(const FName& PropName) const;
-
 private:
 	// Generate list of valid tool names.
 	void RegisterToolNames();
@@ -86,6 +90,9 @@ private:
 	// Handler that implements IToolsetHandler to register with the AI Assistant.
 	TSharedPtr<UE::AIAssistant::FToolsetHandlerTrampoline> RegistryHandler;
 
+	// Generate list of valid tools.
+	TArray<TObjectPtr<UFunction>> GetValidToolMethods() const;
+
 	// Array of valid tool names present in toolset schema.
 	TArray<FString> RegisteredToolNames;
 };
diff --git a/Engine/Plugins/Experimental/AIAssistant/Source/AIAssistant/Private/Tests/AIAssistantToolsetTest.cpp b/Engine/Plugins/Experimental/AIAssistant/Source/AIAssistant/Private/Tests/AIAssistantToolsetTest.cpp
--- a/Engine/Plugins/Experimental/AIAssistant/Source/AIAssistant/Private/Tests/AIAssistantToolsetTest.cpp
+++ b/Engine/Plugins/Experimental/AIAssistant/Source/AIAssistant/Private/Tests/AIAssistantToolsetTest.cpp
@@ -21,6 +21,7 @@
 #if WITH_DEV_AUTOMATION_TESTS
 
 const FString FakeToolsetName(TEXT("FakeToolset"));
+const FString FakeToolset2Name(TEXT("FakeToolset2"));
 const FString FakeToolsetInput(TEXT(R"({"inString":"hello"})"));
 const FString FakeToolName(TEXT("SimpleTool"));
 
@@ -87,15 +88,15 @@ bool FAIAssistantToolsetCanGetValues::RunTest(const FString& UnusedParameters)
 	ToolsetContainer<UFakeToolset> Container;
 	TObjectPtr<UFakeToolset> Instance = *Container;
 	TestEqual("Toolset returns its own name",
-			  Instance->GetToolsetName(), FakeToolsetName);
+			  Instance->GetName(), FakeToolsetName);
 	TestEqual("Toolset returns default non-overridden version",
-			  Instance->GetToolsetVersion(), UAIAssistantToolset::DEFAULT_TOOLSET_VERSION);
-	TestEqual("Toolset returns default non-overridden description",
-			  Instance->GetToolsetDescription(), FString());
+			  Instance->GetVersion(), UAIAssistantToolset::DEFAULT_TOOLSET_VERSION);
+	TestEqual("Toolset returns description from class definition",
+			  Instance->GetDescription(), FString(TEXT("Fake Toolset that does nothing.")));
 	return true;
 }
 
-// Test that overriding toolset properties for version, description works properly.
+// Test that overriding toolset properties for version works properly.
 IMPLEMENT_SIMPLE_AUTOMATION_TEST(
 	FAIAssistantToolsetCanOverrideProperties,
 	"AI.Assistant.Toolset.CanOverrideProperties",
@@ -107,10 +108,13 @@ bool FAIAssistantToolsetCanOverrideProperties::RunTest(const FString& UnusedPara
 
 	ToolsetContainer<UFakeToolset2> Container;
 	TObjectPtr<UFakeToolset2> Instance = *Container;
+	TestEqual("Toolset returns its own name",
+			  Instance->GetName(), FakeToolset2Name);
 	TestEqual("Toolset returns overridden version",
-			  Instance->GetToolsetVersion(), UFakeToolset2::MY_VERSION);
-	TestEqual("Toolset returns overridden description",
-			  Instance->GetToolsetDescription(), UFakeToolset2::MY_DESCRIPTION);
+			  Instance->GetVersion(), UFakeToolset2::MY_VERSION);
+	TestEqual("Toolset returns description from class definition",
+			  Instance->GetDescription(),
+			  FString(TEXT("Fake Toolset that does nothing except update version.")));
 	return true;
 }
 
@@ -118,9 +122,9 @@ TSharedPtr<FJsonObject> GetFakeToolsetSchema(const TObjectPtr<UFakeToolset> Tool
 {
 	FJsonDomBuilder::FObject SchemaBuilder;
 
-	SchemaBuilder.Set(TEXT("name"), Toolset->GetToolsetName());
-	SchemaBuilder.Set(TEXT("version"), Toolset->GetToolsetVersion());
-	SchemaBuilder.Set(TEXT("description"), Toolset->GetToolsetDescription());
+	SchemaBuilder.Set(TEXT("name"), Toolset->GetName());
+	SchemaBuilder.Set(TEXT("version"), Toolset->GetVersion());
+	SchemaBuilder.Set(TEXT("description"), Toolset->GetDescription());
 
 	FJsonDomBuilder::FArray ToolsArray;
 	{
@@ -178,7 +182,7 @@ bool FAIAssistantToolsetGetSchema::RunTest(const FString& UnusedParameters)
 	FJsonSerializer::Serialize(
 		ExpectedSchema.ToSharedRef(),
 		TJsonWriterFactory<TCHAR, TCondensedJsonPrintPolicy<TCHAR>>::Create(&ExpectedSchemaJson));
-	FString OutputSchemaJson = Instance->GetToolsetJsonSchema();
+	FString OutputSchemaJson = Instance->GetJsonSchema();
 	TestEqual("Toolset schema matches expected", OutputSchemaJson, ExpectedSchemaJson);
 	return true;
 }
diff --git a/Engine/Plugins/Experimental/AIAssistant/Source/AIAssistant/Private/Tests/AIAssistantToolsetTest.h b/Engine/Plugins/Experimental/AIAssistant/Source/AIAssistant/Private/Tests/AIAssistantToolsetTest.h
--- a/Engine/Plugins/Experimental/AIAssistant/Source/AIAssistant/Private/Tests/AIAssistantToolsetTest.h
+++ b/Engine/Plugins/Experimental/AIAssistant/Source/AIAssistant/Private/Tests/AIAssistantToolsetTest.h
@@ -10,7 +10,7 @@
 
 /** Fake Toolset that does nothing. */
 UCLASS(Blueprintable)
-class UFakeToolset : public  UAIAssistantToolset
+class UFakeToolset : public UAIAssistantToolset
 {
 	GENERATED_BODY()
 
@@ -22,18 +22,22 @@ public:
 	{
 		return InString;
 	}
+
+	/** UFunction for internal use, should not be marked as tool. */
+	UFUNCTION(meta=(BlueprintInternalUseOnly=true))
+	void SimpleInternalMethod() { }
+
 };
 
-/** Fake Toolset that does nothing except update version and description. */
+/** Fake Toolset that does nothing except update version. */
 UCLASS(Blueprintable)
-class UFakeToolset2 : public  UAIAssistantToolset
+class UFakeToolset2 : public UAIAssistantToolset
 {
 	GENERATED_BODY()
 
 public:
 
 	static inline FString MY_VERSION = FString("1.999.the_best");
-	static inline FString MY_DESCRIPTION = FString("Do nothing number two");
 
 	/** Simple tool that does nothing. */
 	UFUNCTION(BlueprintCallable, Category=Fake)
@@ -42,11 +46,10 @@ public:
 		return InString;
 	}
 
-	// Override the ToolsetVersion, ToolsetDescription properties.
+	// Override the ToolsetVersion property.
 	virtual void PostInitProperties() override
 	{
 		ToolsetVersion = MY_VERSION;
-		ToolsetDescription = MY_DESCRIPTION;
 		Super::PostInitProperties();
 	}
 
