[AI Assistant] Fix tests for toolset schema
Schema generator no longer adds "type": "object" to tool entries

#rnx
#Synced-CL 49128904

diff --git a/Engine/Plugins/Experimental/AIAssistant/Source/AIAssistant/Private/Tests/AIAssistantToolsetTest.cpp b/Engine/Plugins/Experimental/AIAssistant/Source/AIAssistant/Private/Tests/AIAssistantToolsetTest.cpp
--- a/Engine/Plugins/Experimental/AIAssistant/Source/AIAssistant/Private/Tests/AIAssistantToolsetTest.cpp
+++ b/Engine/Plugins/Experimental/AIAssistant/Source/AIAssistant/Private/Tests/AIAssistantToolsetTest.cpp
@@ -129,7 +129,6 @@ TSharedPtr<FJsonObject> GetFakeToolsetSchema(const TObjectPtr<UFakeToolset> Tool
 	FJsonDomBuilder::FArray ToolsArray;
 	{
 		FJsonDomBuilder::FObject ToolEntry;
-		ToolEntry.Set(TEXT("type"), TEXT("object"));
 		ToolEntry.Set(TEXT("description"), TEXT("Simple tool that does nothing."));
 		ToolEntry.Set(TEXT("name"), TEXT("FakeToolset.SimpleTool"));
 		{
