[AIAssistant] Do not double log tool call error messages, and display tool call name.

#Synced-CL 49125772

diff --git a/Engine/Plugins/Experimental/AIAssistant/Source/AIAssistant/Private/AIAssistantToolsetRegistry.cpp b/Engine/Plugins/Experimental/AIAssistant/Source/AIAssistant/Private/AIAssistantToolsetRegistry.cpp
--- a/Engine/Plugins/Experimental/AIAssistant/Source/AIAssistant/Private/AIAssistantToolsetRegistry.cpp
+++ b/Engine/Plugins/Experimental/AIAssistant/Source/AIAssistant/Private/AIAssistantToolsetRegistry.cpp
@@ -27,8 +27,9 @@ namespace UE::AIAssistant
 			ESearchCase::IgnoreCase, ESearchDir::FromEnd) ||
 			Descriptor.ToolsetName.IsEmpty() || Descriptor.ToolName.IsEmpty())
 		{
-			FString Message =
-				TEXT("Invalid tool name format. Expected format: ToolsetName.ToolName");
+			FString Message = FString::Printf(
+				TEXT("Invalid tool name format. Expected format: ToolsetName.ToolName, got %s"),
+				*ToolsetAndToolName);
 			if (LogErrors)
 			{
 				UE_LOG(LogAIAssistant, Error, TEXT("%s"), *Message);
@@ -101,7 +102,7 @@ namespace UE::AIAssistant
 	{
 		// Split the tool name to get the toolset name and the actual tool name.
 		TValueOrError<FToolDescriptor, FString> ToolDescriptor =
-			FToolDescriptor::FromString(ToolsetAndToolName, true);
+			FToolDescriptor::FromString(ToolsetAndToolName);
 		if (ToolDescriptor.HasError())
 		{
 			return MakeFulfilledPromise<TValueOrError<FString, FString>>(
diff --git a/Engine/Plugins/Experimental/AIAssistant/Source/AIAssistant/Private/Tests/AIAssistantToolsetRegistryTest.cpp b/Engine/Plugins/Experimental/AIAssistant/Source/AIAssistant/Private/Tests/AIAssistantToolsetRegistryTest.cpp
--- a/Engine/Plugins/Experimental/AIAssistant/Source/AIAssistant/Private/Tests/AIAssistantToolsetRegistryTest.cpp
+++ b/Engine/Plugins/Experimental/AIAssistant/Source/AIAssistant/Private/Tests/AIAssistantToolsetRegistryTest.cpp
@@ -25,13 +25,11 @@ bool FAIAssistantToolsetDescriptorFromStringValidName::RunTest(const FString& Un
 		FToolDescriptor::FromString(TEXT("ToolsetName.ToolName"));
 	(void)TestTrue(TEXT("FromString_Valid"), Descriptor.HasValue());
 	FToolDescriptor* Value = Descriptor.TryGetValue();
-	if (!Value)
+	if (TestNotNull(TEXT("FromString_Valid_FoundValue"), Value))
 	{
-		return false;
+		(void)TestEqual(TEXT("FromString_Valid_ToolsetName"), Value->ToolsetName, TEXT("ToolsetName"));
+		(void)TestEqual(TEXT("FromString_Valid_ToolName"), Value->ToolName, TEXT("ToolName"));
 	}
-	(void)TestEqual(TEXT("FromString_Valid_ToolsetName"), Value->ToolsetName, TEXT("ToolsetName"));
-	(void)TestEqual(TEXT("FromString_Valid_ToolName"), Value->ToolName, TEXT("ToolName"));
-
 	return true;
 }
 
@@ -47,6 +45,11 @@ bool FAIAssistantToolsetDescriptorFromStringInvalidToolsetName::RunTest(
 	TValueOrError<FToolDescriptor, FString> Descriptor =
 		FToolDescriptor::FromString(TEXT("InvalidToolNameWithoutDot"));
 	(void)TestTrue(TEXT("FromString_Invalid_MissingDot"), Descriptor.HasError());
+	(void)TestEqual(
+		TEXT("FromString_Invalid_MissingDot_String"),
+		Descriptor.GetError(),
+		TEXT("Invalid tool name format. Expected format: ")
+		TEXT("ToolsetName.ToolName, got InvalidToolNameWithoutDot"));
 
 	return true;
 }
@@ -63,6 +66,11 @@ bool FAIAssistantToolsetDescriptorFromStringMissingToolsetName::RunTest(
 	TValueOrError<FToolDescriptor, FString> Descriptor =
 		FToolDescriptor::FromString(TEXT(".ToolNameWithoutToolset"));
 	(void)TestTrue(TEXT("FromString_Invalid_MissingToolsetName"), Descriptor.HasError());
+	(void)TestEqual(
+		TEXT("FromString_Invalid_MissingToolsetName_String"),
+		Descriptor.GetError(),
+		TEXT("Invalid tool name format. Expected format: ")
+		TEXT("ToolsetName.ToolName, got .ToolNameWithoutToolset"));
 
 	return true;
 }
@@ -79,6 +87,11 @@ bool FAIAssistantToolsetDescriptorFromStringMissingToolName::RunTest(
 	TValueOrError<FToolDescriptor, FString> Descriptor =
 		FToolDescriptor::FromString(TEXT("ToolsetWithoutToolname."));
 	(void)TestTrue(TEXT("FromString_Invalid_MissingToolName"), Descriptor.HasError());
+	(void)TestEqual(
+		TEXT("FromString_Invalid_MissingToolName_String"),
+		Descriptor.GetError(),
+		TEXT("Invalid tool name format. Expected format: ")
+		TEXT("ToolsetName.ToolName, got ToolsetWithoutToolname."));
 
 	return true;
 }
