Rename AI assistant config class to be consistent with the rest of the plugin.

#rnx
#Synced-CL 49177557

diff --git a/Engine/Plugins/Experimental/AIAssistant/Source/AIAssistant/Private/AIAssistantConfig.cpp b/Engine/Plugins/Experimental/AIAssistant/Source/AIAssistant/Private/AIAssistantConfig.cpp
--- a/Engine/Plugins/Experimental/AIAssistant/Source/AIAssistant/Private/AIAssistantConfig.cpp
+++ b/Engine/Plugins/Experimental/AIAssistant/Source/AIAssistant/Private/AIAssistantConfig.cpp
@@ -10,74 +10,77 @@
 
 #include "AIAssistantLog.h"
 
-const FString FAIAssistantConfig::DefaultFilename("AIAssistant.json");
+namespace UE::AIAssistant
+{
+	const FString FConfig::DefaultFilename("AIAssistant.json");
 
-const FString FAIAssistantConfig::DefaultMainUrl(
-	"https://dev.epicgames.com/community/assistant/embedded");
+	const FString FConfig::DefaultMainUrl(
+		"https://dev.epicgames.com/community/assistant/embedded");
 
-TArray<FString> FAIAssistantConfig::GetDefaultSearchDirectories()
-{
-	TArray<FString> BaseSearchPaths =
-	{
-		FPaths::EngineConfigDir(),
-		FPaths::EngineUserDir(),
-		FPaths::EngineVersionAgnosticUserDir(),
-	};
-	TArray<FString> AllPathsToSearch;
-	TArray<TOptional<FPaths::EPathConversion>> PathConversions =
-	{
-		TOptional(FPaths::EPathConversion::Engine_NotForLicensees),
-		TOptional(FPaths::EPathConversion::Engine_NoRedist),
-		TOptional(FPaths::EPathConversion::Engine_LimitedAccess),
-		TOptional<FPaths::EPathConversion>(),
-	};
-	for (const TOptional<FPaths::EPathConversion>& PathConversion : PathConversions)
+	TArray<FString> FConfig::GetDefaultSearchDirectories()
 	{
-		for (const FString& BaseSearchPath : BaseSearchPaths)
+		TArray<FString> BaseSearchPaths =
 		{
-			AllPathsToSearch.Emplace(
-				PathConversion.IsSet()
-				? FPaths::ConvertPath(BaseSearchPath, PathConversion.GetValue())
-				: BaseSearchPath);
+			FPaths::EngineConfigDir(),
+			FPaths::EngineUserDir(),
+			FPaths::EngineVersionAgnosticUserDir(),
+		};
+		TArray<FString> AllPathsToSearch;
+		TArray<TOptional<FPaths::EPathConversion>> PathConversions =
+		{
+			TOptional(FPaths::EPathConversion::Engine_NotForLicensees),
+			TOptional(FPaths::EPathConversion::Engine_NoRedist),
+			TOptional(FPaths::EPathConversion::Engine_LimitedAccess),
+			TOptional<FPaths::EPathConversion>(),
+		};
+		for (const TOptional<FPaths::EPathConversion>& PathConversion : PathConversions)
+		{
+			for (const FString& BaseSearchPath : BaseSearchPaths)
+			{
+				AllPathsToSearch.Emplace(
+					PathConversion.IsSet()
+					? FPaths::ConvertPath(BaseSearchPath, PathConversion.GetValue())
+					: BaseSearchPath);
+			}
+
 		}
-		
+		return AllPathsToSearch;
 	}
-	return AllPathsToSearch;
-}
 
-FString FAIAssistantConfig::FindConfigFile(const TArray<FString>& SearchDirectories)
-{
-	for (auto SearchDirectory : SearchDirectories)
+	FString FConfig::FindConfigFile(const TArray<FString>& SearchDirectories)
 	{
-		auto FullFilename = FPaths::Combine(SearchDirectory, DefaultFilename);
-		UE_LOG(
-			LogAIAssistant, Verbose, TEXT("Searching for AI assistant config in %s"),
-			*FullFilename);
-		if (FPaths::FileExists(FullFilename))
+		for (auto SearchDirectory : SearchDirectories)
 		{
-			return FullFilename;
+			auto FullFilename = FPaths::Combine(SearchDirectory, DefaultFilename);
+			UE_LOG(
+				LogAIAssistant, Verbose, TEXT("Searching for AI assistant config in %s"),
+				*FullFilename);
+			if (FPaths::FileExists(FullFilename))
+			{
+				return FullFilename;
+			}
 		}
+		return FString();
 	}
-	return FString();
-}
 
-FAIAssistantConfig FAIAssistantConfig::Load(const FString& Filename)
-{
-	FString Json(TEXT("{}"));
-	if (!Filename.IsEmpty() && !FFileHelper::LoadFileToString(Json, *Filename))
-	{
-		UE_LOG(
-			LogAIAssistant, Error, TEXT("Failed to load AI assistant config from \"%s\"."),
-			*Filename);
-	}
-	FAIAssistantConfig Config;
-	if (!Config.FromJson(Json))
+	FConfig FConfig::Load(const FString& Filename)
 	{
-		UE_LOG(
-			LogAIAssistant, Error, TEXT("Failed to load AI assistant config \"%s\" from \"%s\"."),
-			*Json, *Filename);
-		// Use the default config.
-		(void)Config.FromJson(FString());
+		FString Json(TEXT("{}"));
+		if (!Filename.IsEmpty() && !FFileHelper::LoadFileToString(Json, *Filename))
+		{
+			UE_LOG(
+				LogAIAssistant, Error, TEXT("Failed to load AI assistant config from \"%s\"."),
+				*Filename);
+		}
+		FConfig Config;
+		if (!Config.FromJson(Json))
+		{
+			UE_LOG(
+				LogAIAssistant, Error, TEXT("Failed to load AI assistant config \"%s\" from \"%s\"."),
+				*Json, *Filename);
+			// Use the default config.
+			(void)Config.FromJson(FString());
+		}
+		return Config;
 	}
-	return Config;
 }
\ No newline at end of file
diff --git a/Engine/Plugins/Experimental/AIAssistant/Source/AIAssistant/Private/AIAssistantConfig.h b/Engine/Plugins/Experimental/AIAssistant/Source/AIAssistant/Private/AIAssistantConfig.h
--- a/Engine/Plugins/Experimental/AIAssistant/Source/AIAssistant/Private/AIAssistantConfig.h
+++ b/Engine/Plugins/Experimental/AIAssistant/Source/AIAssistant/Private/AIAssistantConfig.h
@@ -6,48 +6,52 @@
 #include "Serialization/JsonSerializable.h"
 #include "Serialization/JsonSerializerMacros.h"
 
-// Internal configuration for developers of the AI assistant.
-struct FAIAssistantConfig : public FJsonSerializable
+namespace UE::AIAssistant
 {
-	// Initial URL loaded when starting the assistant.
-	// If this isn't specified or is empty, FAIAssistantConfig::DefaultMainUrl is used.
-	FString MainUrl;
-
-	BEGIN_JSON_SERIALIZER
-		JSON_SERIALIZE_WITHDEFAULT("main_url", MainUrl, DefaultMainUrl)
-	END_JSON_SERIALIZER
-
-	// Get the default configuration search directories with the most specific overrides first.
-	//
-	// For example, when searching the engine configuration directory:
-	// * Engine/Restricted/NotForLicensees/Config/
-	// * Engine/Restricted/NoRedist/Config/
-	// * Engine/Restricted/LimitedAccess/Config/
-	// * Engine/Config/
-	//
-	// The base directories that are searched with overrides are:
-	// * Engine configuration directory: Engine/Config/
-	// * Engine directory: Engine/
-	// * Engine user directory (if installed): On Windows %LOCALAPPDATA%/UnrealEngine/Common/
-	static TArray<FString> GetDefaultSearchDirectories();
-
-	// Search for the assistant configuration file in the specified set of directories.
-	static FString FindConfigFile(const TArray<FString>& SearchDirectories);
-
-	// Attempt to the load the configuration from the specified filename. If the filename is empty
-	// return the default configuration. If configuration loading fails, an error is logged and the
-	// default configuration is returned.
-	static FAIAssistantConfig Load(const FString& Filename);
-
-	// Load the configuration from a file by searching default search directories.
-	static FAIAssistantConfig Load()
+	// Internal configuration for developers of the AI assistant.
+	struct FConfig : public FJsonSerializable
 	{
-		return Load(FindConfigFile(GetDefaultSearchDirectories()));
-	}
-
-	// Default initial URL when starting the assistant.
-	static const FString DefaultMainUrl;
-
-	// Configuration filename (basename, not the full path).
-	static const FString DefaultFilename;
-};
+		// Initial URL loaded when starting the assistant.
+		// If this isn't specified or is empty, FConfig::DefaultMainUrl is used.
+		FString MainUrl;
+
+		BEGIN_JSON_SERIALIZER
+			JSON_SERIALIZE_WITHDEFAULT("main_url", MainUrl, DefaultMainUrl)
+		END_JSON_SERIALIZER
+
+		// Get the default configuration search directories with the most specific overrides first.
+		//
+		// For example, when searching the engine configuration directory:
+		// * Engine/Restricted/NotForLicensees/Config/
+		// * Engine/Restricted/NoRedist/Config/
+		// * Engine/Restricted/LimitedAccess/Config/
+		// * Engine/Config/
+		//
+		// The base directories that are searched with overrides are:
+		// * Engine configuration directory: Engine/Config/
+		// * Engine directory: Engine/
+		// * Engine user directory (if installed): On Windows %LOCALAPPDATA%/UnrealEngine/Common/
+		static TArray<FString> GetDefaultSearchDirectories();
+
+		// Search for the assistant configuration file in the specified set of directories.
+		static FString FindConfigFile(const TArray<FString>& SearchDirectories);
+
+		// Attempt to the load the configuration from the specified filename. If the filename is empty
+		// return the default configuration. If configuration loading fails, an error is logged and the
+		// default configuration is returned.
+		static FConfig Load(const FString& Filename);
+
+		// Load the configuration from a file by searching default search directories.
+		static FConfig Load()
+		{
+			return Load(FindConfigFile(GetDefaultSearchDirectories()));
+		}
+
+		// Default initial URL when starting the assistant.
+		static const FString DefaultMainUrl;
+
+		// Configuration filename (basename, not the full path).
+		static const FString DefaultFilename;
+	};
+
+}
\ No newline at end of file
diff --git a/Engine/Plugins/Experimental/AIAssistant/Source/AIAssistant/Private/AIAssistantWebBrowser.cpp b/Engine/Plugins/Experimental/AIAssistant/Source/AIAssistant/Private/AIAssistantWebBrowser.cpp
--- a/Engine/Plugins/Experimental/AIAssistant/Source/AIAssistant/Private/AIAssistantWebBrowser.cpp
+++ b/Engine/Plugins/Experimental/AIAssistant/Source/AIAssistant/Private/AIAssistantWebBrowser.cpp
@@ -28,7 +28,7 @@ using namespace UE::AIAssistant;
 
 void SAIAssistantWebBrowser::Construct(const FArguments& InArgs, const TSharedRef<SDockTab>& InParentTab)
 {
-	Config = FAIAssistantConfig::Load();
+	Config = UE::AIAssistant::FConfig::Load();
 	
 	// Web connection widget.	
 	SAssignNew(WebConnectionWidget, SAIAssistantWebConnectionWidget)
diff --git a/Engine/Plugins/Experimental/AIAssistant/Source/AIAssistant/Private/AIAssistantWebBrowser.h b/Engine/Plugins/Experimental/AIAssistant/Source/AIAssistant/Private/AIAssistantWebBrowser.h
--- a/Engine/Plugins/Experimental/AIAssistant/Source/AIAssistant/Private/AIAssistantWebBrowser.h
+++ b/Engine/Plugins/Experimental/AIAssistant/Source/AIAssistant/Private/AIAssistantWebBrowser.h
@@ -110,7 +110,7 @@ private:
 	TSharedPtr<SWidgetSwitcher> WebBrowserOrWebConnectionSwitcherWidget;
 	
 	// Configuration.
-	FAIAssistantConfig Config;
+	UE::AIAssistant::FConfig Config;
 
 	// Guards MaybeWebApplication.
 	mutable UE::FRecursiveMutex WebApplicationMutex;
diff --git a/Engine/Plugins/Experimental/AIAssistant/Source/AIAssistant/Private/Tests/AIAssistantConfigTest.cpp b/Engine/Plugins/Experimental/AIAssistant/Source/AIAssistant/Private/Tests/AIAssistantConfigTest.cpp
--- a/Engine/Plugins/Experimental/AIAssistant/Source/AIAssistant/Private/Tests/AIAssistantConfigTest.cpp
+++ b/Engine/Plugins/Experimental/AIAssistant/Source/AIAssistant/Private/Tests/AIAssistantConfigTest.cpp
@@ -21,7 +21,7 @@ IMPLEMENT_SIMPLE_AUTOMATION_TEST(
 
 bool FAIAssistantConfigTestHasDefaultSearchDirectories::RunTest(const FString& UnusedParameters)
 {
-	const auto DefaultSearchDirectories = FAIAssistantConfig::GetDefaultSearchDirectories();
+	const auto DefaultSearchDirectories = FConfig::GetDefaultSearchDirectories();
 	(void)TestNotEqual(
 		TEXT("NumDefaultSearchDirectories"), DefaultSearchDirectories.Num(), 0);
 	return true;
@@ -37,11 +37,11 @@ bool FAIAssistantConfigTestFindConfigFileExists::RunTest(const FString& UnusedPa
 	const FTemporaryDirectory TemporaryDirectoryUnused;  // Should be ignored.
 	const FTemporaryDirectory TemporaryDirectoryWithConfig;
 	const FString ConfigFilename =
-		FPaths::Combine(*TemporaryDirectoryWithConfig, FAIAssistantConfig::DefaultFilename);
+		FPaths::Combine(*TemporaryDirectoryWithConfig, FConfig::DefaultFilename);
 	TestTrue(TEXT("WriteConfigFile"), FFileHelper::SaveStringToFile(TEXT(""), *ConfigFilename));
 	(void)TestEqual(
 		TEXT("FindConfigFileExists"),
-		FAIAssistantConfig::FindConfigFile({ *TemporaryDirectoryUnused, *TemporaryDirectoryWithConfig }),
+		FConfig::FindConfigFile({ *TemporaryDirectoryUnused, *TemporaryDirectoryWithConfig }),
 		ConfigFilename);
 	return true;
 }
@@ -56,7 +56,7 @@ bool FAIAssistantConfigTestFindConfigFileMissing::RunTest(const FString& UnusedP
 	const FTemporaryDirectory TemporaryDirectory;
 	(void)TestEqual(
 		TEXT("FindConfigFileMissing"),
-		FAIAssistantConfig::FindConfigFile({ *TemporaryDirectory }),
+		FConfig::FindConfigFile({ *TemporaryDirectory }),
 		"");
 	return true;
 }
@@ -68,8 +68,8 @@ IMPLEMENT_SIMPLE_AUTOMATION_TEST(
 
 bool FAIAssistantConfigTestLoadDefault::RunTest(const FString& UnusedParameters)
 {
-	const FAIAssistantConfig Config = FAIAssistantConfig::Load("");
-	(void)TestEqual("main_url", Config.MainUrl, FAIAssistantConfig::DefaultMainUrl);
+	const FConfig Config = FConfig::Load("");
+	(void)TestEqual("main_url", Config.MainUrl, FConfig::DefaultMainUrl);
 	return true;
 }
 
@@ -81,12 +81,12 @@ IMPLEMENT_SIMPLE_AUTOMATION_TEST(
 bool FAIAssistantConfigTestLoadMainUrl::RunTest(const FString& UnusedParameters)
 {
 	const FTemporaryDirectory TemporaryDirectory;
-	const FString ConfigFilename(FPaths::Combine(*TemporaryDirectory, FAIAssistantConfig::DefaultFilename));
+	const FString ConfigFilename(FPaths::Combine(*TemporaryDirectory, FConfig::DefaultFilename));
 	verify(
 		FFileHelper::SaveStringToFile(
 			FString(TEXT(R"json({"main_url": "https://localhost/assistant"})json")),
 			*ConfigFilename));
-	FAIAssistantConfig Config = FAIAssistantConfig::Load(ConfigFilename);
+	FConfig Config = FConfig::Load(ConfigFilename);
 	(void)TestEqual(TEXT("main_url"), Config.MainUrl, "https://localhost/assistant");
 	return true;
 }
@@ -99,13 +99,13 @@ IMPLEMENT_SIMPLE_AUTOMATION_TEST(
 bool FAIAssistantConfigTestIgnoresAdditionalFields::RunTest(const FString& UnusedParameters)
 {
 	const FTemporaryDirectory TemporaryDirectory;
-	const FString ConfigFilename(FPaths::Combine(*TemporaryDirectory, FAIAssistantConfig::DefaultFilename));
+	const FString ConfigFilename(FPaths::Combine(*TemporaryDirectory, FConfig::DefaultFilename));
 	verify(
 		FFileHelper::SaveStringToFile(
 			FString(TEXT(R"json({"ignore_me": "stuff to ignore"})json")),
 			*ConfigFilename));
-	const FAIAssistantConfig Config = FAIAssistantConfig::Load(ConfigFilename);
-	(void)TestEqual(TEXT("main_url"), Config.MainUrl, FAIAssistantConfig::DefaultMainUrl);
+	const FConfig Config = FConfig::Load(ConfigFilename);
+	(void)TestEqual(TEXT("main_url"), Config.MainUrl, FConfig::DefaultMainUrl);
 	return true;
 }
 
diff --git a/Engine/Plugins/Experimental/AIAssistant/Source/AIAssistant/Private/Tests/AIAssistantWebApplicationTest.cpp b/Engine/Plugins/Experimental/AIAssistant/Source/AIAssistant/Private/Tests/AIAssistantWebApplicationTest.cpp
--- a/Engine/Plugins/Experimental/AIAssistant/Source/AIAssistant/Private/Tests/AIAssistantWebApplicationTest.cpp
+++ b/Engine/Plugins/Experimental/AIAssistant/Source/AIAssistant/Private/Tests/AIAssistantWebApplicationTest.cpp
@@ -91,7 +91,7 @@ namespace UE::AIAssistant
 	struct TWebApplicationWithTracker
 	{
 		TWebApplicationWithTracker(FAutomationTestBase& InTestBase) :
-			Config(FAIAssistantConfig::Load()),
+			Config(FConfig::Load()),
 			WebApplication(MakeShared<FWebApplicationType>(WebApiTracker.CreateFactory())),
 			TestBase(InTestBase)
 		{
@@ -219,7 +219,7 @@ namespace UE::AIAssistant
 		}
 
 		FFakeWebApiTracker WebApiTracker;
-		FAIAssistantConfig Config;
+		FConfig Config;
 		TSharedPtr<FWebApplicationType> WebApplication;
 		FAutomationTestBase& TestBase;
 	};
