[AIAssistant] Do not attempt to run tool calls that do not require a response.

#Synced-CL 49184122

diff --git a/Engine/Plugins/Experimental/AIAssistant/Source/AIAssistant/Private/AIAssistantWebApplication.cpp b/Engine/Plugins/Experimental/AIAssistant/Source/AIAssistant/Private/AIAssistantWebApplication.cpp
--- a/Engine/Plugins/Experimental/AIAssistant/Source/AIAssistant/Private/AIAssistantWebApplication.cpp
+++ b/Engine/Plugins/Experimental/AIAssistant/Source/AIAssistant/Private/AIAssistantWebApplication.cpp
@@ -237,7 +237,10 @@ namespace UE::AIAssistant
 		Algo::TransformIf(
 			Message.MessageContent,
 			ToolCalls,
-			[](const auto& Item) { return Item.ContentType == EMessageContentType::ToolCall; },
+			[](const auto& Item) {
+				const FToolCallContent* ToolCall = Item.Content.template TryGet<FToolCallContent>();
+				return ToolCall && ToolCall->ResponseRequired.Get(false);
+			},
 			[](const auto& Item) { return Item.Content.template Get<FToolCallContent>(); });
 
 		// Process each tool call sequentially.
@@ -268,11 +271,10 @@ namespace UE::AIAssistant
 		const FToolCallContent& ToolCall)
 	{
 		FToolsetRegistry& ToolsetRegistry = GetToolsetRegistry();
-		bool ResponseRequired = ToolCall.ResponseRequired.Get(false);
 		auto ToolCallId = ToolCall.ToolCallId;
 		auto Name = ToolCall.Name;
 		// Tool call must have an ID if a response is required.
-		if (ResponseRequired && !ToolCallId.IsSet())
+		if (!ToolCallId.IsSet())
 		{
 			FString ErrorMessage = FString::Printf(
 				TEXT("Tool call '%s' with arguments '%s' requires a response but has no ID"),
@@ -282,21 +284,15 @@ namespace UE::AIAssistant
 		}
 		// Execute the tool and add the response to the conversation when the tool call completes.
 		return ToolsetRegistry.ExecuteTool(Name, ToolCall.JsonArguments).Next(
-			[WeakThis = SharedThis(this).ToWeakPtr(), ResponseRequired, Name, ToolCallId](
+			[WeakThis = SharedThis(this).ToWeakPtr(), Name, ToolCallId](
 				const TValueOrError<FString, FString>& ToolResult) -> TValueOrError<void, FString>
 			{
 				auto This = WeakThis.Pin();
 				if (!This) return MakeError(WebApplicationShutDownMessage);
 
-				if (ResponseRequired)
-				{
-					This->AddUserMessageToConversation(
-						CreateToolResponseMessage(Name, *ToolCallId, ToolResult));
-				}
-				if (ToolResult.HasError())
-				{
-					return MakeError(ToolResult.GetError());
-				}
+				This->AddUserMessageToConversation(
+					CreateToolResponseMessage(Name, *ToolCallId, ToolResult));
+				if (ToolResult.HasError()) return MakeError(ToolResult.GetError());
 				return MakeValue();
 			});
 	}
@@ -437,4 +433,4 @@ namespace UE::AIAssistant
 		Schema.ToolsetsRawJson = GetToolsetRegistry().GetToolsetJsonSchemas();
 		return AgentEnvironment;
 	}
-}
\ No newline at end of file
+}
diff --git a/Engine/Plugins/Experimental/AIAssistant/Source/AIAssistant/Private/Tests/AIAssistantWebApplicationTest.cpp b/Engine/Plugins/Experimental/AIAssistant/Source/AIAssistant/Private/Tests/AIAssistantWebApplicationTest.cpp
--- a/Engine/Plugins/Experimental/AIAssistant/Source/AIAssistant/Private/Tests/AIAssistantWebApplicationTest.cpp
+++ b/Engine/Plugins/Experimental/AIAssistant/Source/AIAssistant/Private/Tests/AIAssistantWebApplicationTest.cpp
@@ -837,6 +837,52 @@ bool FAIAssistantWebApplicationProcessMessageToolCallMessage::RunTest(
 	return true;
 }
 
+IMPLEMENT_SIMPLE_AUTOMATION_TEST(
+	FAIAssistantWebApplicationProcessMessageOnlyProcessRequiredToolCalls,
+	"AI.Assistant.WebApplication.ProcessMessage.OnlyProcessRequiredToolCalls",
+	AIAssistantTest::Flags);
+
+bool FAIAssistantWebApplicationProcessMessageOnlyProcessRequiredToolCalls::RunTest(
+	const FString& UnusedParameters)
+{
+	FSpyWebApplicationWithTracker WebApplicationWithTracker(*this);
+	auto WebApplication = WebApplicationWithTracker.WebApplication;
+	if (!WebApplicationWithTracker.EnsureInitialized()) return false;
+
+	FToolCallContent PeanutButter = MakeToolCallContent(
+		TEXT("Sandwich.SpreadPeanutButter"),
+		TEXT(R"json({"type": "crunchy"})json"),
+		MakeTuple(1, 1));
+	FToolCallContent Jelly = MakeToolCallContent(
+		TEXT("Sandwich.SpreadJelly"),
+		TEXT(R"json({"flavor": "strawberry"})json"),
+		MakeTuple(1, 2));
+	Jelly.ResponseRequired = false;
+
+	FMessage AgentToolCallMessage = MakeAgentToolCallMessage(1, { PeanutButter, Jelly });
+
+	WebApplication->bCallBaseProcessToolCallContent = false;
+	WebApplication->ProcessToolCallContentFakeResults.Enqueue(
+		MakeShared<TPromise<TValueOrError<void, FString>>>(
+			MakeFulfilledPromise<TValueOrError<void, FString>>(MakeValue())));
+	WebApplication->ProcessToolCallContentFakeResults.Enqueue(
+		MakeShared<TPromise<TValueOrError<void, FString>>>(
+			MakeFulfilledPromise<TValueOrError<void, FString>>(MakeValue())));
+	WebApplication->ProcessMessage(AgentToolCallMessage);
+
+	// Check that the tool that did not require a resopnse was not processed.
+	if (TestEqual(
+		TEXT("ProcessMessageWithToolCalls_ProcessToolCallContentCount"),
+		WebApplication->ProcessToolCallContentCalls.Num(), 1))
+	{
+		(void)TestEqual(
+			TEXT("ProcessMessageWithToolCalls_ProcessToolCallContentCount"),
+			WebApplication->ProcessToolCallContentCalls[0], PeanutButter);
+	}
+
+	return true;
+}
+
 IMPLEMENT_SIMPLE_AUTOMATION_TEST(
 	FAIAssistantWebApplicationProcessMessageProcessOnlyOnce,
 	"AI.Assistant.WebApplication.ProcessMessage.ProcessOnlyOnce",
