Remove compare-against-self since the comment to squelch it for PVS studio doesn't seem to be working.

#Synced-CL 49186506

diff --git a/Engine/Plugins/Experimental/AIAssistant/Source/AIAssistant/Private/Tests/AIAssistantWebApiTest.cpp b/Engine/Plugins/Experimental/AIAssistant/Source/AIAssistant/Private/Tests/AIAssistantWebApiTest.cpp
--- a/Engine/Plugins/Experimental/AIAssistant/Source/AIAssistant/Private/Tests/AIAssistantWebApiTest.cpp
+++ b/Engine/Plugins/Experimental/AIAssistant/Source/AIAssistant/Private/Tests/AIAssistantWebApiTest.cpp
@@ -222,8 +222,6 @@ void AIAssistantWebApiMessageId::Define()
 			It("returns true when the object is identical", [this]
 				{
 					FMessageId MessageIdSame = MakeMessageId(1);
-					// Intentional self-compare in test.
-					TestTrue(TEXT("MessageId should be equal to itself"), MessageId == MessageId); //-V501
 					TestTrue(TEXT("MessageIds created with same index should be equal"),
 						MessageId == MessageIdSame);
 				});
@@ -283,9 +281,6 @@ void AIAssistantWebApiConversationId::Define()
 			It("returns true when the object is identical", [this]
 				{
 					FConversationId ConversationIdSame = MakeConversationId(1);
-					// Intentional self-compare in test.
-					TestTrue(TEXT("ConversationId should be equal to itself"),
-						ConversationId == ConversationId); //-V501
 					TestTrue(TEXT("ConversationIds with same Id should be equal"),
 						ConversationId == ConversationIdSame);
 				});
