[AI Assistant] Fix JavaScript error exception handling.

JavaScript Error objects convert to "{}" when passed to JSON.stringify()
which was resulting in error context being lost when being passed from
JavaScript to C++. This resolves the issue by using Error.toString() if an
error object is detected when handling an error condition.

#Synced-CL 49212909

diff --git a/Engine/Plugins/Experimental/AIAssistant/Source/AIAssistant/Private/AIAssistantWebApi.cpp b/Engine/Plugins/Experimental/AIAssistant/Source/AIAssistant/Private/AIAssistantWebApi.cpp
--- a/Engine/Plugins/Experimental/AIAssistant/Source/AIAssistant/Private/AIAssistantWebApi.cpp
+++ b/Engine/Plugins/Experimental/AIAssistant/Source/AIAssistant/Private/AIAssistantWebApi.cpp
@@ -47,6 +47,9 @@ try {
 }
 )js");
 
+	const FString FWebApi::ConvertErrorToString =
+		TEXT(R"js(error instanceof Error ? error.toString() : error)js");
+
 	FWebApi::FWebApi(
 		IWebJavaScriptExecutor& JavaScriptExecutor,
 		IWebJavaScriptDelegateBinder& JavaScriptDelegateBinder) :
@@ -203,7 +206,7 @@ try {
 				WebJavaScriptResultDelegate->FormatJavaScriptHandler(HandlerId);
 			return TPair<FString, FString>(
 				FString::Format(*HandlerFormat, { TEXT("result"), TEXT("false") }),
-				FString::Format(*HandlerFormat, { TEXT("error"), TEXT("true") }));
+				FString::Format(*HandlerFormat, { *ConvertErrorToString, TEXT("true") }));
 		}
 		return TPair<FString, FString>();
 	}
diff --git a/Engine/Plugins/Experimental/AIAssistant/Source/AIAssistant/Private/AIAssistantWebApi.h b/Engine/Plugins/Experimental/AIAssistant/Source/AIAssistant/Private/AIAssistantWebApi.h
--- a/Engine/Plugins/Experimental/AIAssistant/Source/AIAssistant/Private/AIAssistantWebApi.h
+++ b/Engine/Plugins/Experimental/AIAssistant/Source/AIAssistant/Private/AIAssistantWebApi.h
@@ -573,6 +573,10 @@ namespace UE::AIAssistant
 		//   a function error.
 		static const FString FunctionCallFormatTemplate;
 		static const FString FunctionCallbackFormatTemplate;
+
+		// JavaScript that converts a variable called "error" into a string or returns the
+		// original object.
+		static const FString ConvertErrorToString;
 	};
 }
 
diff --git a/Engine/Plugins/Experimental/AIAssistant/Source/AIAssistant/Private/Tests/AIAssistantWebApiTest.cpp b/Engine/Plugins/Experimental/AIAssistant/Source/AIAssistant/Private/Tests/AIAssistantWebApiTest.cpp
--- a/Engine/Plugins/Experimental/AIAssistant/Source/AIAssistant/Private/Tests/AIAssistantWebApiTest.cpp
+++ b/Engine/Plugins/Experimental/AIAssistant/Source/AIAssistant/Private/Tests/AIAssistantWebApiTest.cpp
@@ -61,6 +61,11 @@ namespace UE::AIAssistant
 		{
 			return *WebApi.WebJavaScriptResultDelegate;
 		}
+
+		static const FString& GetConvertErrorToString()
+		{
+			return FWebApi::ConvertErrorToString;
+		}
 	};
 }
 
@@ -420,7 +425,7 @@ bool FAIAssistantWebApiTestFormatResultAndErrorHandlers::RunTest(
 		Handlers.Value,
 		FString::Format(
 			*JavaScriptResultDelegate.FormatJavaScriptHandler(FakeHandlerId),
-			{ TEXT("error"), TEXT("true") }));
+			{ *FWebApiAccessor::GetConvertErrorToString(), TEXT("true")}));
 	return true;
 }
 
