[AI Assistant] RAII container for delegate handles.

Simplifies unregistration of delegate handlers from non-UObject classes.

#Synced-CL 49212909

diff --git a/Engine/Plugins/Experimental/AIAssistant/Source/AIAssistant/Private/AIAssistantDelegateHandle.h b/Engine/Plugins/Experimental/AIAssistant/Source/AIAssistant/Private/AIAssistantDelegateHandle.h
new file mode 100644
--- /dev/null
+++ b/Engine/Plugins/Experimental/AIAssistant/Source/AIAssistant/Private/AIAssistantDelegateHandle.h
@@ -0,0 +1,87 @@
+// Copyright Epic Games, Inc. All Rights Reserved.
+#pragma once
+
+#include "Delegates/IDelegateInstance.h"
+#include "Templates/UniquePtr.h"
+#include "Templates/UnrealTemplate.h"
+
+namespace UE::AIAssistant
+{
+	// Delegate handle RAII container.
+	class FDelegateHandleRaiiBase
+	{
+	public:
+		virtual ~FDelegateHandleRaiiBase() { Reset(); }
+
+		FDelegateHandleRaiiBase(const FDelegateHandleRaiiBase&) = delete;
+		FDelegateHandleRaiiBase& operator=(const FDelegateHandleRaiiBase&) = delete;
+
+		// Reset the handle.
+		virtual void Reset() { DelegateHandle.Reset(); }
+
+		// Get the underlying delegate handle.
+		const FDelegateHandle& Get() const { return DelegateHandle; }
+
+	protected:
+		FDelegateHandleRaiiBase(FDelegateHandle&& Handle) : DelegateHandle(MoveTemp(Handle)) {}
+
+	private:
+		FDelegateHandle DelegateHandle;
+	};
+
+	// RAII container for a delegate handle.
+	template<typename DelegateRegistrationT>
+	class TDelegateHandleRaii : public FDelegateHandleRaiiBase
+	{
+	public:
+		// Move the delegate handle for a registration into a container.
+		TDelegateHandleRaii(DelegateRegistrationT& Registration, FDelegateHandle&& Handle) :
+			FDelegateHandleRaiiBase(MoveTemp(Handle)),
+			DelegateRegistration(Registration)
+		{
+		}
+
+		// On destruction, unregister the handle from the delegate registration.
+		virtual ~TDelegateHandleRaii() { Reset(); }
+
+		// Unregister the handle from the delegate registration.
+		virtual void Reset() override
+		{
+			const auto& Handle = Get();
+			if (Handle.IsValid())
+			{
+				DelegateRegistration.Remove(Handle);
+				FDelegateHandleRaiiBase::Reset();
+			}
+		}
+
+	private:
+		DelegateRegistrationT& DelegateRegistration;
+	};
+
+	class FDelegateHandleRaii : public TUniquePtr<FDelegateHandleRaiiBase>
+	{
+	public:
+		using TUniquePtr<FDelegateHandleRaiiBase>::TUniquePtr;
+
+		// Whether this references a valid handle.
+		bool IsValid() const
+		{
+			return TUniquePtr::IsValid() && Get()->Get().IsValid();
+		}
+
+	public:
+		// Create an instance of a delegate handle RAII container.
+		//
+		// For example:
+		// auto Handle = FDelegateHandleRaii::Create(
+		//      OnSomeMulticastDelegate.AddLambda([]() -> void { /* do something */ }));
+		template<typename DelegateRegistrationT>
+		static FDelegateHandleRaii Create(
+			DelegateRegistrationT& Registration, FDelegateHandle&& Handle)
+		{
+			return MakeUnique<TDelegateHandleRaii<DelegateRegistrationT>>(
+				Registration, MoveTemp(Handle));
+		}
+	};
+}
\ No newline at end of file
diff --git a/Engine/Plugins/Experimental/AIAssistant/Source/AIAssistant/Private/Tests/AIAssistantDelegateHandleTest.cpp b/Engine/Plugins/Experimental/AIAssistant/Source/AIAssistant/Private/Tests/AIAssistantDelegateHandleTest.cpp
new file mode 100644
--- /dev/null
+++ b/Engine/Plugins/Experimental/AIAssistant/Source/AIAssistant/Private/Tests/AIAssistantDelegateHandleTest.cpp
@@ -0,0 +1,137 @@
+// Copyright Epic Games, Inc. All Rights Reserved.
+#pragma once
+
+#include "Containers/UnrealString.h"
+#include "Delegates/DelegateCombinations.h"
+#include "Misc/AutomationTest.h"
+#include "Templates/Function.h"
+#include "Templates/UniquePtr.h"
+#include "Templates/UnrealTemplate.h"
+
+#include "AIAssistantDelegateHandle.h"
+#include "AIAssistantTestFlags.h"
+
+using namespace UE::AIAssistant;
+
+namespace
+{
+	struct FCallCounter
+	{
+		int CallCount = 0;
+
+		void Call()
+		{
+			CallCount++;
+		}
+
+		TFunction<void()> Callable()
+		{
+			return [this]() -> void { Call(); };
+		}
+	};
+}
+
+DECLARE_MULTICAST_DELEGATE(FAIAssistantOnTest)
+
+BEGIN_DEFINE_SPEC(AIAssistantDelegateHandleRaii, "AI.Assistant.DelegateHandleRaii",
+	AIAssistantTest::Flags)
+END_DEFINE_SPEC(AIAssistantDelegateHandleRaii)
+
+void AIAssistantDelegateHandleRaii::Define()
+{
+	Describe(TEXT("TDelegateHandleRaii"), [this]
+		{
+			It(TEXT("can construct a container attached to a delegate handle"), [this]
+				{
+					FAIAssistantOnTest OnTest;
+					FCallCounter CallCounter;
+					auto RawDelegateHandle = OnTest.AddLambda(CallCounter.Callable());
+					auto RawDelegateHandleToMove = RawDelegateHandle;
+					TDelegateHandleRaii Handle(OnTest, MoveTemp(RawDelegateHandleToMove));
+					(void)TestEqual(TEXT("Delegate handle"), Handle.Get(), RawDelegateHandle);
+				});
+
+			It(TEXT("can construct a container that references an invalid handle"), [this]
+				{
+					FAIAssistantOnTest OnTest;
+					FDelegateHandle RawDelegateHandle;
+					TDelegateHandleRaii Handle(OnTest, MoveTemp(RawDelegateHandle));
+					(void)TestFalse(TEXT("Delegate handle is not valid"), Handle.Get().IsValid());
+				});
+
+			It(TEXT("should unregister a delegate handler on reset"), [this]
+				{
+					FAIAssistantOnTest OnTest;
+					FCallCounter CallCounter;
+					TDelegateHandleRaii Handle(OnTest, OnTest.AddLambda(CallCounter.Callable()));
+					OnTest.Broadcast();
+					(void)TestEqual(TEXT("Called once"), CallCounter.CallCount, 1);
+					Handle.Reset();
+
+					OnTest.Broadcast();
+					(void)TestEqual(TEXT("Still called once"), CallCounter.CallCount, 1);
+					(void)TestFalse(
+						TEXT("Delegate handle is not valid"),
+						Handle.Get().IsValid());
+				});
+
+			It(TEXT("should unregister a delegate handler on destruction"), [this]
+				{
+					FAIAssistantOnTest OnTest;
+					FCallCounter CallCounter;
+					{
+						TDelegateHandleRaii Handle(
+							OnTest, OnTest.AddLambda(CallCounter.Callable()));
+						OnTest.Broadcast();
+						(void)TestEqual(TEXT("Called once"), CallCounter.CallCount, 1);
+					}
+
+					OnTest.Broadcast();
+					(void)TestEqual(TEXT("Still called once"), CallCounter.CallCount, 1);
+				});
+		});
+
+	Describe(TEXT("FDelegateHandleRaii::Create"), [this]
+		{
+			It(TEXT("should create a container attached to a delegate handle"), [this]
+				{
+					FAIAssistantOnTest OnTest;
+					FCallCounter CallCounter;
+					auto RawDelegateHandle = OnTest.AddLambda(CallCounter.Callable());
+					auto RawDelegateHandleToMove = RawDelegateHandle;
+					auto Handle = FDelegateHandleRaii::Create(
+						OnTest, MoveTemp(RawDelegateHandleToMove));
+					if (!TestTrue(TEXT("Handle created"), Handle.IsValid())) return;
+					(void)TestEqual(TEXT("Delegate handle"),
+						Handle->Get(), RawDelegateHandle);
+				});
+		});
+
+	Describe(TEXT("FDelegateHandleRaii::IsValid"), [this]
+		{
+			It(TEXT("should return false if it does not reference a valid pointer"), [this]
+				{
+					FDelegateHandleRaii Handle;
+					(void)TestFalse(TEXT("Not valid on construction"), Handle.IsValid());
+				});
+
+			It(TEXT("should return true if it references a valid pointer and handle"), [this]
+				{
+					FAIAssistantOnTest OnTest;
+					FCallCounter CallCounter;
+					auto Handle = FDelegateHandleRaii::Create(
+						OnTest, OnTest.AddLambda(CallCounter.Callable()));
+					(void)TestTrue(TEXT("Valid after create"), Handle.IsValid());
+				});
+
+			It(TEXT("should return false if it references an invalid handle"), [this]
+				{
+					FAIAssistantOnTest OnTest;
+					FCallCounter CallCounter;
+					auto Handle = FDelegateHandleRaii::Create(
+						OnTest, OnTest.AddLambda(CallCounter.Callable()));
+					Handle->Reset();
+					(void)TestFalse(TEXT("Not valid after reset"), Handle.IsValid());
+				});
+		});
+}
\ No newline at end of file
