[AI Assistant]: Fix delegate handle test build in release

#Synced-CL 49230250

diff --git a/Engine/Plugins/Experimental/AIAssistant/Source/AIAssistant/Private/Tests/AIAssistantDelegateHandleTest.cpp b/Engine/Plugins/Experimental/AIAssistant/Source/AIAssistant/Private/Tests/AIAssistantDelegateHandleTest.cpp
--- a/Engine/Plugins/Experimental/AIAssistant/Source/AIAssistant/Private/Tests/AIAssistantDelegateHandleTest.cpp
+++ b/Engine/Plugins/Experimental/AIAssistant/Source/AIAssistant/Private/Tests/AIAssistantDelegateHandleTest.cpp
@@ -13,6 +13,8 @@
 
 using namespace UE::AIAssistant;
 
+#if WITH_DEV_AUTOMATION_TESTS
+
 namespace
 {
 	struct FCallCounter
@@ -134,4 +136,6 @@ void AIAssistantDelegateHandleRaii::Define()
 					(void)TestFalse(TEXT("Not valid after reset"), Handle.IsValid());
 				});
 		});
-}
\ No newline at end of file
+}
+
+#endif  // WITH_DEV_AUTOMATION_TESTS
