[AIAssistant] Return correct result from OnConversationUpdate callback

#Synced-CL 49214499

diff --git a/Engine/Plugins/Experimental/AIAssistant/Source/AIAssistant/Private/AIAssistantWebApplication.cpp b/Engine/Plugins/Experimental/AIAssistant/Source/AIAssistant/Private/AIAssistantWebApplication.cpp
--- a/Engine/Plugins/Experimental/AIAssistant/Source/AIAssistant/Private/AIAssistantWebApplication.cpp
+++ b/Engine/Plugins/Experimental/AIAssistant/Source/AIAssistant/Private/AIAssistantWebApplication.cpp
@@ -334,7 +334,7 @@ namespace UE::AIAssistant
 								return This && !This->HandleErrorResult(
 									TEXT("OnConversationUpdate"), Result);
 							});
-						return true;
+						return false;
 					}).Next([WeakThis](auto HandleOrError) {
 						auto This = WeakThis.Pin();
 						if (!This) return;
