[AI Assistant] Stop pending JavaScript callbacks on PreExit event.

#Synced-CL 49232378

diff --git a/Engine/Plugins/Experimental/AIAssistant/Source/AIAssistant/Private/AIAssistantWebJavaScriptResultDelegate.cpp b/Engine/Plugins/Experimental/AIAssistant/Source/AIAssistant/Private/AIAssistantWebJavaScriptResultDelegate.cpp
--- a/Engine/Plugins/Experimental/AIAssistant/Source/AIAssistant/Private/AIAssistantWebJavaScriptResultDelegate.cpp
+++ b/Engine/Plugins/Experimental/AIAssistant/Source/AIAssistant/Private/AIAssistantWebJavaScriptResultDelegate.cpp
@@ -14,16 +14,22 @@ const FString UAIAssistantWebJavaScriptResultDelegate::CanceledError(
 	TEXT(R"json("canceled")json"));
 
 void UAIAssistantWebJavaScriptResultDelegate::Bind(
-	UE::AIAssistant::IWebJavaScriptDelegateBinder& WebJavaScriptDelegateBinder)
+	UE::AIAssistant::IWebJavaScriptDelegateBinder& WebJavaScriptDelegateBinder,
+	FSimpleMulticastDelegate& UnbindDelegate)
 {
 	ScopedWebJavaScriptDelegateBinder.Emplace(
 		WebJavaScriptDelegateBinder, GetName(), this, true /* bIsPermanent */);
+	OnPreExitDelegateHandle = UE::AIAssistant::FDelegateHandleRaii::Create(
+		UnbindDelegate,
+		UnbindDelegate.AddUObject(
+			this, &UAIAssistantWebJavaScriptResultDelegate::Unbind));
 }
 
 void UAIAssistantWebJavaScriptResultDelegate::Unbind()
 {
 	ScopedWebJavaScriptDelegateBinder.Reset();
 	CompleteAllPendingPromises();
+	OnPreExitDelegateHandle.Reset();
 }
 
 FString UAIAssistantWebJavaScriptResultDelegate::CreateHandlerId()
diff --git a/Engine/Plugins/Experimental/AIAssistant/Source/AIAssistant/Private/AIAssistantWebJavaScriptResultDelegate.h b/Engine/Plugins/Experimental/AIAssistant/Source/AIAssistant/Private/AIAssistantWebJavaScriptResultDelegate.h
--- a/Engine/Plugins/Experimental/AIAssistant/Source/AIAssistant/Private/AIAssistantWebJavaScriptResultDelegate.h
+++ b/Engine/Plugins/Experimental/AIAssistant/Source/AIAssistant/Private/AIAssistantWebJavaScriptResultDelegate.h
@@ -4,6 +4,7 @@
 #include "Async/Mutex.h"
 #include "Async/Future.h"
 #include "Containers/UnrealString.h"
+#include "Misc/CoreDelegates.h"
 #include "Misc/Guid.h"
 #include "Templates/Function.h"
 #include "Templates/SharedPointer.h"
@@ -11,6 +12,7 @@
 #include "Templates/UnrealTemplate.h"
 #include "UObject/Object.h"
 
+#include "AIAssistantDelegateHandle.h"
 #include "AIAssistantWebJavaScriptDelegateBinder.h"
 
 #include "AIAssistantWebJavaScriptResultDelegate.generated.h"
@@ -58,7 +60,9 @@ public:
 
 public:
 	// Bind this object to a JavaScript execution environment as Name ("aiassistantresultdelegate").
-	void Bind(UE::AIAssistant::IWebJavaScriptDelegateBinder& WebJavaScriptDelegateBinder);
+	void Bind(
+		UE::AIAssistant::IWebJavaScriptDelegateBinder& WebJavaScriptDelegateBinder,
+		FSimpleMulticastDelegate& UnbindDelegate = FCoreDelegates::OnPreExit);
 
 	// Unregister this object from the previous JavaScript delegate binder.
 	// NOTE: This must be called before the bound IWebJavaScriptDelegateBinder is destroyed. It's not
@@ -127,6 +131,8 @@ private:
 	UE::FMutex PendingPromisesByIdLock; // Guards PendingPromisesById
 	TMap<FString, TPromise<FResult>> PendingPromisesById;
 
+	UE::AIAssistant::FDelegateHandleRaii OnPreExitDelegateHandle;
+
 private:
 	// Create a handler ID.
 	static FString CreateHandlerId();
diff --git a/Engine/Plugins/Experimental/AIAssistant/Source/AIAssistant/Private/Tests/AIAssistantWebJavaScriptResultDelegateTest.cpp b/Engine/Plugins/Experimental/AIAssistant/Source/AIAssistant/Private/Tests/AIAssistantWebJavaScriptResultDelegateTest.cpp
--- a/Engine/Plugins/Experimental/AIAssistant/Source/AIAssistant/Private/Tests/AIAssistantWebJavaScriptResultDelegateTest.cpp
+++ b/Engine/Plugins/Experimental/AIAssistant/Source/AIAssistant/Private/Tests/AIAssistantWebJavaScriptResultDelegateTest.cpp
@@ -3,6 +3,7 @@
 #include "Async/Future.h"
 #include "Containers/StaticArray.h"
 #include "Containers/UnrealString.h"
+#include "Delegates/DelegateCombinations.h"
 #include "Misc/AutomationTest.h"
 #include "Templates/Tuple.h"
 #include "UObject/Object.h"
@@ -322,4 +323,35 @@ bool FAIAssistantWebJavaScriptResultDelegateTestCompleteAllPendingPromises::RunT
 	return true;
 }
 
+DECLARE_MULTICAST_DELEGATE(FAIAssistantTestOnPreExit)
+
+IMPLEMENT_SIMPLE_AUTOMATION_TEST(
+	FAIAssistantWebJavaScriptResultDelegateTestCompleteAllPendingPromisesOnPreExit,
+	"AI.Assistant.WebJavaScriptResultDelegate.CompleteAllPendingPromisesOnPreExit",
+	AIAssistantTest::Flags);
+
+bool FAIAssistantWebJavaScriptResultDelegateTestCompleteAllPendingPromisesOnPreExit::RunTest(
+	const FString& UnusedParameters)
+{
+	FAIAssistantTestOnPreExit OnPreExit;
+	FFakeWebJavaScriptDelegateBinder Binder;
+	TStrongObjectPtr<UAIAssistantWebJavaScriptResultDelegate> ResultDelegate(
+		NewObject<UAIAssistantWebJavaScriptResultDelegate>());
+	ResultDelegate->Bind(Binder, OnPreExit);
+	TPair<FString, TFuture<UAIAssistantWebJavaScriptResultDelegate::FResult>> ResultHandle =
+		ResultDelegate->RegisterResultHandlerForFuture();
+
+	// Use OnPreExit() to force clean up of pending promises.
+	OnPreExit.Broadcast();
+
+	(void)TestTrue(TEXT("FutureIsReady"), ResultHandle.Value.IsReady());
+	auto Result = ResultHandle.Value.Consume();
+	(void)TestEqual(
+		TEXT("FutureIsCanceled"),
+		UAIAssistantWebJavaScriptResultDelegate::CanceledError,
+		Result.Json);
+	(void)TestTrue(TEXT("FutureHasError"), Result.bJsonIsError);
+	return true;
+}
+
 #endif  // WITH_DEV_AUTOMATION_TESTS
\ No newline at end of file
