[AI Assistant] Move module and web browser to shared configuration

#Synced-CL 49280788

diff --git a/Engine/Plugins/Experimental/AIAssistant/Source/AIAssistant/Private/AIAssistant.cpp b/Engine/Plugins/Experimental/AIAssistant/Source/AIAssistant/Private/AIAssistant.cpp
--- a/Engine/Plugins/Experimental/AIAssistant/Source/AIAssistant/Private/AIAssistant.cpp
+++ b/Engine/Plugins/Experimental/AIAssistant/Source/AIAssistant/Private/AIAssistant.cpp
@@ -14,6 +14,7 @@
 #include "Widgets/Docking/SDockTab.h"
 
 #include "AIAssistantCommands.h"
+#include "AIAssistantCurrentConfig.h"
 #include "AIAssistantInputProcessor.h"
 #include "AIAssistantSlateQuerier.h"
 #include "AIAssistantStyle.h"
@@ -38,13 +39,8 @@ static const FName AIAssistantTabName("AIAssistant");
 
 /*virtual*/ void FAIAssistantModule::StartupModule()
 {
-	GConfig->GetBool(TEXT("AIAssistant"), TEXT("bIsEnabled"), bIsAIAssistantEnabled, GEditorIni);
-
-	if (!bIsAIAssistantEnabled)
-	{
-		return;
-	}
-	
+	CurrentConfig = MakeShared<UE::AIAssistant::FCurrentConfig>();
+	if (!CurrentConfig->IsEnabled()) return;
 
 	FAIAssistantStyle::Initialize();
 	FAIAssistantStyle::ReloadTextures();
@@ -104,11 +100,7 @@ static const FName AIAssistantTabName("AIAssistant");
 
 /*virtual*/ void FAIAssistantModule::ShutdownModule()
 {
-	if (!bIsAIAssistantEnabled)
-	{
-		return;
-	}
-
+	if (!CurrentConfig->IsEnabled()) return;
 	
 	if (FSlateApplication::IsInitialized())
 	{
@@ -317,7 +309,8 @@ TSharedRef<SDockTab> FAIAssistantModule::OnSpawnPluginTab(const FSpawnTabArgs& S
 
 	// Work around for the web browser drawing out of bounds while doing the on demand clipping
 	AIAssistantWebBrowserWidget = SNew(SAIAssistantWebBrowser, DockTabWidget)
-		.Clipping(EWidgetClipping::ClipToBounds);
+		.Clipping(EWidgetClipping::ClipToBounds)
+		.CurrentConfig(CurrentConfig);
 
 	DockTabWidget->SetContent(AIAssistantWebBrowserWidget.ToSharedRef());
 
diff --git a/Engine/Plugins/Experimental/AIAssistant/Source/AIAssistant/Private/AIAssistant.h b/Engine/Plugins/Experimental/AIAssistant/Source/AIAssistant/Private/AIAssistant.h
--- a/Engine/Plugins/Experimental/AIAssistant/Source/AIAssistant/Private/AIAssistant.h
+++ b/Engine/Plugins/Experimental/AIAssistant/Source/AIAssistant/Private/AIAssistant.h
@@ -10,6 +10,8 @@
 #include "Templates/SharedPointer.h"
 
 
+#include "AIAssistantCurrentConfig.h"
+
 class FAIAssistantInputProcessor;
 class FSpawnTabArgs;
 class SAIAssistantWebBrowser;
@@ -84,10 +86,8 @@ private:
 	void GenerateStatusBarPanelDrawerSummon(TArray<UStatusBarSubsystem::FTabIdAndButtonLabel>& OutTabIdsAndLabels, const TSharedRef<SDockTab>& InParentTab) const;
 
 	TSharedRef<SDockTab> OnSpawnPluginTab(const FSpawnTabArgs& SpawnTabArgs);
-
-
-	bool bIsAIAssistantEnabled = true;
 	
+	TSharedPtr<UE::AIAssistant::FCurrentConfig> CurrentConfig;
 	TSharedPtr<FUICommandList> PluginCommands;
 	TSharedPtr<FAIAssistantInputProcessor> InputProcessor;
 	TSharedPtr<SAIAssistantWebBrowser> AIAssistantWebBrowserWidget;
diff --git a/Engine/Plugins/Experimental/AIAssistant/Source/AIAssistant/Private/AIAssistantWebBrowser.cpp b/Engine/Plugins/Experimental/AIAssistant/Source/AIAssistant/Private/AIAssistantWebBrowser.cpp
--- a/Engine/Plugins/Experimental/AIAssistant/Source/AIAssistant/Private/AIAssistantWebBrowser.cpp
+++ b/Engine/Plugins/Experimental/AIAssistant/Source/AIAssistant/Private/AIAssistantWebBrowser.cpp
@@ -28,7 +28,10 @@ using namespace UE::AIAssistant;
 
 void SAIAssistantWebBrowser::Construct(const FArguments& InArgs, const TSharedRef<SDockTab>& InParentTab)
 {
-	Config = UE::AIAssistant::FConfig::Load();
+	verify(InArgs._CurrentConfig.IsValid());  // Configuration must be set.
+	CurrentConfig = InArgs._CurrentConfig;
+
+	auto Config = CurrentConfig->Load();
 	
 	// Web connection widget.	
 	SAssignNew(WebConnectionWidget, SAIAssistantWebConnectionWidget)
@@ -38,14 +41,14 @@ void SAIAssistantWebBrowser::Construct(const FArguments& InArgs, const TSharedRe
 			{
 				return GetConnectionState();
 			})
-		.OnReconnect_Lambda([this]() -> void
+		.OnReconnect_Lambda([this, Config]() -> void
 			{
 				if (WebBrowserWidget.IsValid())
 				{
 					// Must reload to (1) clear any error page, (2) make sure the correct page is
 					// actually loaded.
 					InitializeWebApplication();
-					LoadUrl(Config.MainUrl, EOpenBrowserMode::Embedded);
+					LoadUrl(Config->MainUrl, EOpenBrowserMode::Embedded);
 				}
 			});
 
@@ -76,7 +79,7 @@ void SAIAssistantWebBrowser::Construct(const FArguments& InArgs, const TSharedRe
 	// Web Browser.
 	SAssignNew(WebBrowserWidget, SWebBrowser, /*passed to SWebBrowser ctr..*/ WebBrowserWindow)
 #if UE_AIA_SET_INITIAL_URL_VIA_WEB_BROWSER_SLATE_ARG
-		.InitialURL(Config.MainUrl) 
+		.InitialURL(Config->MainUrl) 
 #endif
 		.ShowControls(false)
 		.ShowErrorMessage(false) // ..suppress showing error pages
@@ -193,7 +196,7 @@ void SAIAssistantWebBrowser::Construct(const FArguments& InArgs, const TSharedRe
 	
 #if !UE_AIA_SET_INITIAL_URL_VIA_WEB_BROWSER_SLATE_ARG
 	InitializeWebApplication();
-	LoadUrl(Config.MainUrl, EOpenBrowserMode::Embedded);
+	LoadUrl(Config->MainUrl, EOpenBrowserMode::Embedded);
 #endif
 }
 
diff --git a/Engine/Plugins/Experimental/AIAssistant/Source/AIAssistant/Private/AIAssistantWebBrowser.h b/Engine/Plugins/Experimental/AIAssistant/Source/AIAssistant/Private/AIAssistantWebBrowser.h
--- a/Engine/Plugins/Experimental/AIAssistant/Source/AIAssistant/Private/AIAssistantWebBrowser.h
+++ b/Engine/Plugins/Experimental/AIAssistant/Source/AIAssistant/Private/AIAssistantWebBrowser.h
@@ -10,7 +10,7 @@
 #include "SWebBrowser.h"
 #include "Widgets/Layout/SWidgetSwitcher.h"
 
-#include "AIAssistantConfig.h"
+#include "AIAssistantCurrentConfig.h"
 #include "AIAssistantWebApplication.h"
 #include "AIAssistantWebConnectionWidget.h"
 #include "AIAssistantWebJavaScriptDelegateBinder.h"
@@ -47,6 +47,8 @@ class SAIAssistantWebBrowser :
 {
 public:
 	SLATE_BEGIN_ARGS(SAIAssistantWebBrowser) {}
+		// This *must* be set to construct this widget.
+		SLATE_ARGUMENT(TSharedPtr<UE::AIAssistant::FCurrentConfig>, CurrentConfig);
 	SLATE_END_ARGS()
 
 	void Construct(const FArguments& InArgs, const TSharedRef<class SDockTab>& InParentTab);
@@ -110,7 +112,7 @@ private:
 	TSharedPtr<SWidgetSwitcher> WebBrowserOrWebConnectionSwitcherWidget;
 	
 	// Configuration.
-	UE::AIAssistant::FConfig Config;
+	TSharedPtr<UE::AIAssistant::FCurrentConfig> CurrentConfig;
 
 	// Guards MaybeWebApplication.
 	mutable UE::FRecursiveMutex WebApplicationMutex;
