[AI Assistant] Add OnPreExit event handling to AI assistant config

This provides a central location to handle the editor OnPreExit event,
track whether it has been fired and inject a replacement multicast
delegate in tests.

#Synced-CL 49340993

diff --git a/Engine/Plugins/Experimental/AIAssistant/Source/AIAssistant/Private/AIAssistantCurrentConfig.cpp b/Engine/Plugins/Experimental/AIAssistant/Source/AIAssistant/Private/AIAssistantCurrentConfig.cpp
--- a/Engine/Plugins/Experimental/AIAssistant/Source/AIAssistant/Private/AIAssistantCurrentConfig.cpp
+++ b/Engine/Plugins/Experimental/AIAssistant/Source/AIAssistant/Private/AIAssistantCurrentConfig.cpp
@@ -9,12 +9,18 @@
 
 namespace UE::AIAssistant
 {
-	FCurrentConfig::FCurrentConfig(const FString& EditorIniFilename)
+	FCurrentConfig::FCurrentConfig(const FConstructorArgs& Args) :
+		bIsExiting(false),
+		OnPreExitDelegate(*Args.OnPreExit),
+		OnPreExitDelegateHandle(
+			FDelegateHandleRaii::Create(
+				OnPreExitDelegate,
+				OnPreExitDelegate.AddRaw(this, &FCurrentConfig::OnPreExitHandler)))
 	{
-		GConfig->GetBool(TEXT("AIAssistant"), TEXT("bIsEnabled"), bIsEnabled, EditorIniFilename);
+		GConfig->GetBool(
+			TEXT("AIAssistant"), TEXT("bIsEnabled"), bIsEnabled, Args.EditorIniFilename);
 	}
 
-	// Load or reload the static configuration.
 	TSharedPtr<const FConfig> FCurrentConfig::Load()
 	{
 		auto LoadedConfig = MakeShared<const FConfig>(FConfig::Load());
@@ -22,11 +28,12 @@ namespace UE::AIAssistant
 		return LoadedConfig;
 	}
 
-	// Get the currently loaded configuration.
 	TSharedPtr<const FConfig> FCurrentConfig::GetOrLoad()
 	{
 		auto LoadedConfig = Config;
 		if (LoadedConfig) return LoadedConfig;
 		return Load();
 	}
+
+	void FCurrentConfig::OnPreExitHandler() { bIsExiting.store(true); }
 }
\ No newline at end of file
diff --git a/Engine/Plugins/Experimental/AIAssistant/Source/AIAssistant/Private/AIAssistantCurrentConfig.h b/Engine/Plugins/Experimental/AIAssistant/Source/AIAssistant/Private/AIAssistantCurrentConfig.h
--- a/Engine/Plugins/Experimental/AIAssistant/Source/AIAssistant/Private/AIAssistantCurrentConfig.h
+++ b/Engine/Plugins/Experimental/AIAssistant/Source/AIAssistant/Private/AIAssistantCurrentConfig.h
@@ -1,10 +1,14 @@
 // Copyright Epic Games, Inc. All Rights Reserved.
 #pragma once
 
+#include <atomic>
+
 #include "CoreGlobals.h"
-#include "Templates/SharedPointer.h"
+#include "Delegates/Delegate.h"
+#include "Misc/CoreDelegates.h"
 
 #include "AIAssistantConfig.h"
+#include "AIAssistantDelegateHandle.h"
 
 namespace UE::AIAssistant
 {
@@ -15,10 +19,23 @@ namespace UE::AIAssistant
 	// for restricted / unrestricted and language changes but they should probably be moved
 	// into this class to make it easier to test that functionality without modifying the
 	// global state of the editor.
-	class FCurrentConfig : public TSharedFromThis<FCurrentConfig, ESPMode::ThreadSafe>
+	class FCurrentConfig
 	{
 	public:
-		FCurrentConfig(const FString& EditorIniFilename = GEditorIni);
+		// Arguments for the constructor so that the caller can override arguments without
+		// repeating default values.
+		struct FConstructorArgs
+		{
+			// Explicitly define a constructor to workaround Clang DR 1351/1397.
+			FConstructorArgs() {} // = default;
+			FString EditorIniFilename = GEditorIni;
+			// IMPORTANT: The provided delegate lifetime must be longer than this object and 
+			// FCurrentConfig constructed with this object.
+			FSimpleMulticastDelegate* OnPreExit = &FCoreDelegates::OnPreExit;
+		};
+
+	public:
+		FCurrentConfig(const FConstructorArgs& Args = FConstructorArgs());
 		~FCurrentConfig() = default;
 
 		// Load or reload the static configuration.
@@ -31,8 +48,21 @@ namespace UE::AIAssistant
 		// This is set on construction and never modified for the lifetime of this instance.
 		bool IsEnabled() const { return bIsEnabled;  }
 
+		// Multicast delegate that can be used to watch for application pre-exit events.
+		FSimpleMulticastDelegate& OnPreExit() const { return OnPreExitDelegate; }
+
+		// Returns true if the application is exiting.
+		bool IsExiting() const { return bIsExiting.load();  }
+
+	private:
+		// Handles OnPreExit event.
+		void OnPreExitHandler();
+
 	private:
 		TSharedPtr<const FConfig> Config;
+		std::atomic<bool> bIsExiting;
+		FSimpleMulticastDelegate& OnPreExitDelegate;
+		FDelegateHandleRaii OnPreExitDelegateHandle;
 		bool bIsEnabled = true;
 	};
 
diff --git a/Engine/Plugins/Experimental/AIAssistant/Source/AIAssistant/Private/Tests/AIAssistantCurrentConfigTest.cpp b/Engine/Plugins/Experimental/AIAssistant/Source/AIAssistant/Private/Tests/AIAssistantCurrentConfigTest.cpp
--- a/Engine/Plugins/Experimental/AIAssistant/Source/AIAssistant/Private/Tests/AIAssistantCurrentConfigTest.cpp
+++ b/Engine/Plugins/Experimental/AIAssistant/Source/AIAssistant/Private/Tests/AIAssistantCurrentConfigTest.cpp
@@ -1,7 +1,9 @@
 // Copyright Epic Games, Inc. All Rights Reserved.
 
 #include "Containers/UnrealString.h"
+#include "Delegates/Delegate.h"
 #include "Misc/AutomationTest.h"
+#include "Misc/CoreDelegates.h"
 #include "Misc/FileHelper.h"
 #include "Misc/Paths.h"
 #include "Templates/Function.h"
@@ -112,8 +114,11 @@ void AIAssistantCurrentConfig::Define()
 					EditorIniFile.With(
 						*this,
 						[this](const FString& EditorIniFilename) -> void {
-							FCurrentConfig CurrentConfig(EditorIniFilename);
-							(void)TestTrue(TEXT("Assistant is enabled"), CurrentConfig.IsEnabled());
+							FCurrentConfig::FConstructorArgs ConstructorArgs;
+							ConstructorArgs.EditorIniFilename = EditorIniFilename;
+							FCurrentConfig CurrentConfig(ConstructorArgs);
+							(void)TestTrue(
+								TEXT("Assistant is enabled"), CurrentConfig.IsEnabled());
 						});
 				});
 
@@ -123,12 +128,54 @@ void AIAssistantCurrentConfig::Define()
 					EditorIniFile.With(
 						*this,
 						[this](const FString& EditorIniFilename) -> void {
-							FCurrentConfig CurrentConfig(EditorIniFilename);
-							(void)TestFalse(TEXT("Assistant is disable"), CurrentConfig.IsEnabled());
+							FCurrentConfig::FConstructorArgs ConstructorArgs;
+							ConstructorArgs.EditorIniFilename = EditorIniFilename;
+							FCurrentConfig CurrentConfig(ConstructorArgs);
+							(void)TestFalse(
+								TEXT("Assistant is disable"), CurrentConfig.IsEnabled());
 						});
 				});
 		});
-}
 
+	Describe(TEXT("OnPreExit"), [this]
+		{
+			It(TEXT("Should default to the application OnPreExit delegate"), [this]
+				{
+					FCurrentConfig CurrentConfig;
+					(void)TestEqual(
+						TEXT("OnPreExit"),
+						&CurrentConfig.OnPreExit(),
+						&FCoreDelegates::OnPreExit);
+				});
+			
+			It(TEXT("Should be overridable"), [this]
+				{
+					FSimpleMulticastDelegate OnTest;
+					FCurrentConfig::FConstructorArgs ConstructorArgs;
+					ConstructorArgs.OnPreExit = &OnTest;
+					FCurrentConfig CurrentConfig(ConstructorArgs);
+					(void)TestEqual(TEXT("OnPreExit"), &CurrentConfig.OnPreExit(), &OnTest);
+				});
+		});
+
+	Describe(TEXT("IsExiting"), [this]
+		{
+			It(TEXT("Should default to false"), [this]
+				{
+					FCurrentConfig CurrentConfig;
+					(void)TestFalse(TEXT("IsExiting"), CurrentConfig.IsExiting());
+				});
+
+			It(TEXT("Should be true after OnPreExit event"), [this]
+				{
+					FSimpleMulticastDelegate OnTest;
+					FCurrentConfig::FConstructorArgs ConstructorArgs;
+					ConstructorArgs.OnPreExit = &OnTest;
+					FCurrentConfig CurrentConfig(ConstructorArgs);
+					OnTest.Broadcast();
+					(void)TestTrue(TEXT("IsExiting"), CurrentConfig.IsExiting());
+				});
+		});
+}
 
 #endif
\ No newline at end of file
