[AIAssistant] Change the return value of registered delegates to be an enum describing the outcome instead of a bool.

Boolean return values don't communicate meaning very well. Having an enum makes it clear what the outcome represents, and makes it harder to return the wrong thing.

#Synced-CL 49214499

diff --git a/Engine/Plugins/Experimental/AIAssistant/Source/AIAssistant/Private/AIAssistantWebApi.cpp b/Engine/Plugins/Experimental/AIAssistant/Source/AIAssistant/Private/AIAssistantWebApi.cpp
--- a/Engine/Plugins/Experimental/AIAssistant/Source/AIAssistant/Private/AIAssistantWebApi.cpp
+++ b/Engine/Plugins/Experimental/AIAssistant/Source/AIAssistant/Private/AIAssistantWebApi.cpp
@@ -123,7 +123,7 @@ try {
 
 	TFuture<TValueOrError<FWebApi::FConversationUpdateHandle, FString>>
 		FWebApi::RegisterOnConversationUpdate(
-		TFunction<bool(FConversationId)> OnConversationUpdate)
+		FOnConversationUpdateCallback OnConversationUpdate)
 	{
 		using FResult = UAIAssistantWebJavaScriptResultDelegate::FResult;
 		FString HandlerId = RegisterCallback(
@@ -134,7 +134,7 @@ try {
 					UE_LOG(
 						LogAIAssistant, Error, TEXT("Javascript Error running callback: %s"),
 						*Result.Json);
-					return true;
+					return UAIAssistantWebJavaScriptResultDelegate::EHandlerStatus::Unregister;
 				}
 				return OnConversationUpdate(Id);
 			});
@@ -220,7 +220,8 @@ try {
 	}
 
 	FString FWebApi::RegisterCallback(
-		TFunction<bool(UAIAssistantWebJavaScriptResultDelegate::FResult)> Callback)
+		TFunction<UAIAssistantWebJavaScriptResultDelegate::EHandlerStatus(
+			UAIAssistantWebJavaScriptResultDelegate::FResult)> Callback)
 	{
 		return WebJavaScriptResultDelegate->RegisterResultHandlerForCallback(Callback);
 	}
diff --git a/Engine/Plugins/Experimental/AIAssistant/Source/AIAssistant/Private/AIAssistantWebApi.h b/Engine/Plugins/Experimental/AIAssistant/Source/AIAssistant/Private/AIAssistantWebApi.h
--- a/Engine/Plugins/Experimental/AIAssistant/Source/AIAssistant/Private/AIAssistantWebApi.h
+++ b/Engine/Plugins/Experimental/AIAssistant/Source/AIAssistant/Private/AIAssistantWebApi.h
@@ -413,10 +413,13 @@ namespace UE::AIAssistant
 		// Opaque handle to data needed to unregister conversation update callback.
 		using FConversationUpdateHandle = TTuple<FString, FString>;
 
-		// Register for conversation updates. Return true to unregister, or explicitly call
-		// UnregisterOnConversationUpdate.
+		using FOnConversationUpdateCallback = 
+			TFunction<UAIAssistantWebJavaScriptResultDelegate::EHandlerStatus(FConversationId)>;
+
+		// Register for conversation updates. To unregister the callback return
+		// EHandlerStatus::Unregister or explicitly call UnregisterOnConversationUpdate.
 		TFuture<TValueOrError<FConversationUpdateHandle, FString>> RegisterOnConversationUpdate(
-			TFunction<bool(FConversationId)> OnConversationUpdate);
+			FOnConversationUpdateCallback OnConversationUpdate);
 
 		// Unregister conversation update callback using the CallbackRegistration returned from
 		// RegisterOnConversationUpdate.
@@ -443,7 +446,8 @@ namespace UE::AIAssistant
 		// times from Javascript. This is virtual so that the returned handler ID can be captured
 		// in tests.
 		virtual FString RegisterCallback(
-			TFunction<bool(UAIAssistantWebJavaScriptResultDelegate::FResult)> Callback);
+			TFunction<UAIAssistantWebJavaScriptResultDelegate::EHandlerStatus(
+				UAIAssistantWebJavaScriptResultDelegate::FResult)> Callback);
 
 		// Execute a javascript function handling the result with the specified handler.
 		// NOTE: This is virtual so tests can capture the function, arguments and handler ID before
diff --git a/Engine/Plugins/Experimental/AIAssistant/Source/AIAssistant/Private/AIAssistantWebApplication.cpp b/Engine/Plugins/Experimental/AIAssistant/Source/AIAssistant/Private/AIAssistantWebApplication.cpp
--- a/Engine/Plugins/Experimental/AIAssistant/Source/AIAssistant/Private/AIAssistantWebApplication.cpp
+++ b/Engine/Plugins/Experimental/AIAssistant/Source/AIAssistant/Private/AIAssistantWebApplication.cpp
@@ -323,10 +323,15 @@ namespace UE::AIAssistant
 				if (!This) return;
 
 				WebApi.RegisterOnConversationUpdate(
-					[WeakThis](FConversationId Id) -> bool
+					[WeakThis](FConversationId Id) ->
+						UAIAssistantWebJavaScriptResultDelegate::EHandlerStatus
 					{
 						auto This = WeakThis.Pin();
-						if (!This) return false;
+						if (!This)
+						{
+							return
+								UAIAssistantWebJavaScriptResultDelegate::EHandlerStatus::Unregister;
+						}
 
 						This->OnConversationUpdate(Id).Next(
 							[WeakThis](const auto& Result) {
@@ -334,7 +339,7 @@ namespace UE::AIAssistant
 								return This && !This->HandleErrorResult(
 									TEXT("OnConversationUpdate"), Result);
 							});
-						return false;
+						return UAIAssistantWebJavaScriptResultDelegate::EHandlerStatus::Continue;
 					}).Next([WeakThis](auto HandleOrError) {
 						auto This = WeakThis.Pin();
 						if (!This) return;
diff --git a/Engine/Plugins/Experimental/AIAssistant/Source/AIAssistant/Private/AIAssistantWebJavaScriptResultDelegate.cpp b/Engine/Plugins/Experimental/AIAssistant/Source/AIAssistant/Private/AIAssistantWebJavaScriptResultDelegate.cpp
--- a/Engine/Plugins/Experimental/AIAssistant/Source/AIAssistant/Private/AIAssistantWebJavaScriptResultDelegate.cpp
+++ b/Engine/Plugins/Experimental/AIAssistant/Source/AIAssistant/Private/AIAssistantWebJavaScriptResultDelegate.cpp
@@ -63,21 +63,22 @@ UAIAssistantWebJavaScriptResultDelegate::RegisterResultHandlerForFuture()
 	FString HandlerId = CreateHandlerId();
 	TFuture<FResult> Future = CreatePromiseAndGetFuture(HandlerId);
 	RegisterResultHandlerWithId(
-		[this](FResultHandlerContext&& ResultHandlerContext) mutable -> bool
+		[this](FResultHandlerContext&& ResultHandlerContext) mutable -> EHandlerStatus
 		{
 			TryCompletePendingPromise(MoveTemp(ResultHandlerContext));
-			return true;
+			return EHandlerStatus::Unregister;
 		},
 		HandlerId);
 	return TPair<FString, TFuture<FResult>>(MoveTemp(HandlerId), MoveTemp(Future));
 }
 
 FString UAIAssistantWebJavaScriptResultDelegate::RegisterResultHandlerForCallback(
-	TFunction<bool(FResult)> Callback)
+	TFunction<EHandlerStatus(FResult)> Callback)
 {
 	FString HandlerId = CreateHandlerId();
 	RegisterResultHandlerWithId(
-		[Callback=MoveTemp(Callback)](FResultHandlerContext ResultHandlerContext) mutable -> bool
+		[Callback=MoveTemp(Callback)](FResultHandlerContext ResultHandlerContext) mutable ->
+			UAIAssistantWebJavaScriptResultDelegate::EHandlerStatus
 		{
 			return Callback(MoveTemp(ResultHandlerContext));
 		},
@@ -95,9 +96,9 @@ void UAIAssistantWebJavaScriptResultDelegate::HandleResult(
 		if (!HandlerPtrPtr) return;
 		Handler = *HandlerPtrPtr;
 	}
-	bool RemoveHandler =
+	EHandlerStatus Result =
 		(*Handler)(FResultHandlerContext{ { ResultJson, bResultJsonIsError }, HandlerId });
-	if (RemoveHandler)
+	if (Result == EHandlerStatus::Unregister)
 	{
 		UnregisterResultHandler(HandlerId);
 	}
diff --git a/Engine/Plugins/Experimental/AIAssistant/Source/AIAssistant/Private/AIAssistantWebJavaScriptResultDelegate.h b/Engine/Plugins/Experimental/AIAssistant/Source/AIAssistant/Private/AIAssistantWebJavaScriptResultDelegate.h
--- a/Engine/Plugins/Experimental/AIAssistant/Source/AIAssistant/Private/AIAssistantWebJavaScriptResultDelegate.h
+++ b/Engine/Plugins/Experimental/AIAssistant/Source/AIAssistant/Private/AIAssistantWebJavaScriptResultDelegate.h
@@ -53,10 +53,17 @@ public:
 		FString HandlerId;
 	};
 
+	// Whether or not the handler should remain registered or be unregistered after running.
+	enum class EHandlerStatus
+	{
+		Continue,
+		Unregister
+	};
+
 	// Handles a JavaScript execution result.
-	// This functor should return true to unregister the handler or false to leave the handler
-	// registered for future callbacks.
-	using FResultHandler = TFunction<bool(FResultHandlerContext&&)>;
+	// This functor should return EHandlerStatus::Unregister to unregister the handler or
+	// EHandlerStatus::Continue to leave the handler registered for future callbacks.
+	using FResultHandler = TFunction<EHandlerStatus(FResultHandlerContext&&)>;
 
 public:
 	// Bind this object to a JavaScript execution environment as Name ("aiassistantresultdelegate").
@@ -78,9 +85,9 @@ public:
 	// This handler will always only be executed once.
 	TPair<FString, TFuture<FResult>> RegisterResultHandlerForFuture();
 
-	// Register a callback that may be called multiple times. Return true from the callback to
-	// unregister the handler, or explicitly call UnregisterResultHandler.
-	FString RegisterResultHandlerForCallback(TFunction<bool(FResult)> Callback);
+	// Register a callback that may be called multiple times. Return EHandlerStatus::Unregister
+	// from the callback to unregister the handler, or explicitly call UnregisterResultHandler.
+	FString RegisterResultHandlerForCallback(TFunction<EHandlerStatus(FResult)> Callback);
 
 	// Unregister a previously registered handler with the specified ID.
 	void UnregisterResultHandler(const FString& HandlerId);
diff --git a/Engine/Plugins/Experimental/AIAssistant/Source/AIAssistant/Private/Tests/AIAssistantFakeWebApi.cpp b/Engine/Plugins/Experimental/AIAssistant/Source/AIAssistant/Private/Tests/AIAssistantFakeWebApi.cpp
--- a/Engine/Plugins/Experimental/AIAssistant/Source/AIAssistant/Private/Tests/AIAssistantFakeWebApi.cpp
+++ b/Engine/Plugins/Experimental/AIAssistant/Source/AIAssistant/Private/Tests/AIAssistantFakeWebApi.cpp
@@ -88,7 +88,8 @@ namespace UE::AIAssistant
 	}
 
 	FString FFakeWebApi::RegisterCallback(
-		TFunction<bool(UAIAssistantWebJavaScriptResultDelegate::FResult)> Callback)
+		TFunction<UAIAssistantWebJavaScriptResultDelegate::EHandlerStatus(
+			UAIAssistantWebJavaScriptResultDelegate::FResult)> Callback)
 	{
 		RegisteredCallbackHandlerIds.Add(FWebApi::RegisterCallback(MoveTemp(Callback)));
 		return RegisteredCallbackHandlerIds.Last();
diff --git a/Engine/Plugins/Experimental/AIAssistant/Source/AIAssistant/Private/Tests/AIAssistantFakeWebApi.h b/Engine/Plugins/Experimental/AIAssistant/Source/AIAssistant/Private/Tests/AIAssistantFakeWebApi.h
--- a/Engine/Plugins/Experimental/AIAssistant/Source/AIAssistant/Private/Tests/AIAssistantFakeWebApi.h
+++ b/Engine/Plugins/Experimental/AIAssistant/Source/AIAssistant/Private/Tests/AIAssistantFakeWebApi.h
@@ -118,7 +118,8 @@ namespace UE::AIAssistant
 			const TCHAR* HandlerId) override;
 
 		FString RegisterCallback(
-			TFunction<bool(UAIAssistantWebJavaScriptResultDelegate::FResult)> Callback) override;
+			TFunction<UAIAssistantWebJavaScriptResultDelegate::EHandlerStatus(
+				UAIAssistantWebJavaScriptResultDelegate::FResult)> Callback) override;
 
 		// If possible, get the value of a result as JSON.  If the value cannot be converted to
 		// JSON return an empty string.
diff --git a/Engine/Plugins/Experimental/AIAssistant/Source/AIAssistant/Private/Tests/AIAssistantWebApiTest.cpp b/Engine/Plugins/Experimental/AIAssistant/Source/AIAssistant/Private/Tests/AIAssistantWebApiTest.cpp
--- a/Engine/Plugins/Experimental/AIAssistant/Source/AIAssistant/Private/Tests/AIAssistantWebApiTest.cpp
+++ b/Engine/Plugins/Experimental/AIAssistant/Source/AIAssistant/Private/Tests/AIAssistantWebApiTest.cpp
@@ -671,7 +671,7 @@ bool FAIAssistantWebApiTestRegisterOnConversationUpdate::RunTest(const FString&
 		WebApi->RegisterOnConversationUpdate([&CallCount](FConversationId ConversationId)
 		{
 			CallCount += 1;
-			return false;
+			return UAIAssistantWebJavaScriptResultDelegate::EHandlerStatus::Continue;
 		});
 	const FString& HandlerId = WebApi->RegisteredCallbackHandlerIds.Last();
 	return
@@ -701,7 +701,7 @@ bool FAIAssistantWebApiTestRegisterOnConversationUpdateAndInvoke::RunTest(
 		{
 			ActualConversationId = ConversationId.Id;
 			CallCount += 1;
-			return false;
+			return UAIAssistantWebJavaScriptResultDelegate::EHandlerStatus::Continue;
 		});
 	const FString& HandlerId = WebApi->RegisteredCallbackHandlerIds.Last();
 	if (!WebApi->TestExpectAsyncFunctionCallAndComplete(
@@ -747,7 +747,7 @@ bool FAIAssistantWebApiTestRegisterOnConversationUpdateAndInvokeTwice::RunTest(
 		{
 			ActualConversationIds.Push(ConversationId.Id);
 			CallCount += 1;
-			return false;
+			return UAIAssistantWebJavaScriptResultDelegate::EHandlerStatus::Continue;
 		});
 
 	const FString& HandlerId = WebApi->RegisteredCallbackHandlerIds.Last();
@@ -799,7 +799,7 @@ bool FAIAssistantWebApiTestUnegisterOnConversationUpdateViaReturnValue::RunTest(
 		{
 			ActualConversationIds.Push(ConversationId.Id);
 			CallCount += 1;
-			return true;
+			return UAIAssistantWebJavaScriptResultDelegate::EHandlerStatus::Unregister;
 		});
 	const FString& HandlerId = WebApi->RegisteredCallbackHandlerIds.Last();
 	if (!WebApi->TestExpectAsyncFunctionCallAndComplete(
@@ -848,7 +848,7 @@ bool FAIAssistantWebApiTestUnregisterOnConversationUpdate::RunTest(const FString
 		{
 			ActualConversationIds.Push(ConversationId.Id);
 			CallCount += 1;
-			return false;
+			return UAIAssistantWebJavaScriptResultDelegate::EHandlerStatus::Continue;
 		});
 	const FString& HandlerId = WebApi->RegisteredCallbackHandlerIds.Last();
 	if (!WebApi->TestExpectAsyncFunctionCallAndComplete(
diff --git a/Engine/Plugins/Experimental/AIAssistant/Source/AIAssistant/Private/Tests/AIAssistantWebJavaScriptResultDelegateTest.cpp b/Engine/Plugins/Experimental/AIAssistant/Source/AIAssistant/Private/Tests/AIAssistantWebJavaScriptResultDelegateTest.cpp
--- a/Engine/Plugins/Experimental/AIAssistant/Source/AIAssistant/Private/Tests/AIAssistantWebJavaScriptResultDelegateTest.cpp
+++ b/Engine/Plugins/Experimental/AIAssistant/Source/AIAssistant/Private/Tests/AIAssistantWebJavaScriptResultDelegateTest.cpp
@@ -75,9 +75,10 @@ bool FAIAssistantWebJavaScriptResultDelegateTestRegisterResultHandler::RunTest(
 		NewObject<UAIAssistantWebJavaScriptResultDelegate>());
 
 	FString HandlerId = ResultDelegate->RegisterResultHandler(
-		[](UAIAssistantWebJavaScriptResultDelegate::FResultHandlerContext&&) -> bool
+		[](UAIAssistantWebJavaScriptResultDelegate::FResultHandlerContext&&) ->
+			UAIAssistantWebJavaScriptResultDelegate::EHandlerStatus
 		{
-			return true;
+			return UAIAssistantWebJavaScriptResultDelegate::EHandlerStatus::Unregister;
 		});
 	(void)TestNotEqual(TEXT("HandlerId"), TEXT(""), HandlerId);
 	return true;
@@ -101,7 +102,7 @@ bool FAIAssistantWebJavaScriptResultDelegateTestCallResultHandler::RunTest(
 		FString HandlerId = ResultDelegate->RegisterResultHandler(
 			[this, &ExpectedJsonResult, bExpectedJsonResultIsError, &bCalledHandler](
 				UAIAssistantWebJavaScriptResultDelegate::FResultHandlerContext&&
-				ResultHandlerContext) -> bool
+				ResultHandlerContext) -> UAIAssistantWebJavaScriptResultDelegate::EHandlerStatus
 			{
 				(void)TestEqual(
 					TEXT("JsonResult"), ExpectedJsonResult,
@@ -110,7 +111,7 @@ bool FAIAssistantWebJavaScriptResultDelegateTestCallResultHandler::RunTest(
 					TEXT("bJsonResultIsError"), bExpectedJsonResultIsError,
 					ResultHandlerContext.bJsonIsError);
 				bCalledHandler = true;
-				return true;
+				return UAIAssistantWebJavaScriptResultDelegate::EHandlerStatus::Unregister;
 			});
 
 		FWebJavaScriptResultDelegateAccessor::CallHandleResult(
@@ -133,10 +134,11 @@ bool FAIAssistantWebJavaScriptResultDelegateTestCallResultHandlerOnce::RunTest(
 
 	bool bWasCalled = false;
 	FString HandlerId = ResultDelegate->RegisterResultHandler(
-		[&bWasCalled](UAIAssistantWebJavaScriptResultDelegate::FResultHandlerContext&&) -> bool
+		[&bWasCalled](UAIAssistantWebJavaScriptResultDelegate::FResultHandlerContext&&) ->
+			UAIAssistantWebJavaScriptResultDelegate::EHandlerStatus
 		{
 			bWasCalled = true;
-			return true;
+			return UAIAssistantWebJavaScriptResultDelegate::EHandlerStatus::Unregister;
 		});
 	(void)TestEqual(
 		TEXT("NumberOfHandlers"), 1,
@@ -167,10 +169,11 @@ bool FAIAssistantWebJavaScriptResultDelegateTestCallResultHandlerMultiple::RunTe
 	int NumberOfCalls = 0;
 	FString HandlerId = ResultDelegate->RegisterResultHandler(
 		[&NumberOfCalls](
-			UAIAssistantWebJavaScriptResultDelegate::FResultHandlerContext&&) mutable -> bool
+			UAIAssistantWebJavaScriptResultDelegate::FResultHandlerContext&&) mutable ->
+			UAIAssistantWebJavaScriptResultDelegate::EHandlerStatus
 		{
 			NumberOfCalls++;
-			return false;
+			return UAIAssistantWebJavaScriptResultDelegate::EHandlerStatus::Continue;
 		});
 
 	for (int i = 0; i < 2; ++i)
@@ -250,10 +253,11 @@ bool FAIAssistantWebJavaScriptResultDelegateTestRegisterResultHandlerForCallback
 		NewObject<UAIAssistantWebJavaScriptResultDelegate>());
 	int CallCount = 0;
 	FString HandlerId = ResultDelegate->RegisterResultHandlerForCallback(
-		[&CallCount](UAIAssistantWebJavaScriptResultDelegate::FResult) mutable -> bool
+		[&CallCount](UAIAssistantWebJavaScriptResultDelegate::FResult) mutable ->
+			UAIAssistantWebJavaScriptResultDelegate::EHandlerStatus
 		{
 			CallCount++;
-			return false;
+			return UAIAssistantWebJavaScriptResultDelegate::EHandlerStatus::Continue;
 		});
 	const int ExpectedCalls = 3;
 	for (int i = 0; i < ExpectedCalls; ++i)
@@ -278,10 +282,11 @@ bool FAIAssistantWebJavaScriptResultDelegateTestUnregisterResultHandler::RunTest
 		NewObject<UAIAssistantWebJavaScriptResultDelegate>());
 	bool bWasCalled = false;
 	FString HandlerId = ResultDelegate->RegisterResultHandler(
-		[&bWasCalled](UAIAssistantWebJavaScriptResultDelegate::FResultHandlerContext&&) -> bool
+		[&bWasCalled](UAIAssistantWebJavaScriptResultDelegate::FResultHandlerContext&&)
+			-> UAIAssistantWebJavaScriptResultDelegate::EHandlerStatus
 		{
 			bWasCalled = true;
-			return true;
+			return UAIAssistantWebJavaScriptResultDelegate::EHandlerStatus::Unregister;
 		});
 	ResultDelegate->UnregisterResultHandler(HandlerId);
 	(void)TestEqual(
