[AI Assistant] Fix web application shutdown on editor exit

Stop creating / using the web application when the editor is
shutting down.

#Synced-CL 49353738

diff --git a/Engine/Plugins/Experimental/AIAssistant/Source/AIAssistant/Private/AIAssistantWebApi.cpp b/Engine/Plugins/Experimental/AIAssistant/Source/AIAssistant/Private/AIAssistantWebApi.cpp
--- a/Engine/Plugins/Experimental/AIAssistant/Source/AIAssistant/Private/AIAssistantWebApi.cpp
+++ b/Engine/Plugins/Experimental/AIAssistant/Source/AIAssistant/Private/AIAssistantWebApi.cpp
@@ -52,12 +52,13 @@ try {
 
 	FWebApi::FWebApi(
 		IWebJavaScriptExecutor& JavaScriptExecutor,
-		IWebJavaScriptDelegateBinder& JavaScriptDelegateBinder) :
+		IWebJavaScriptDelegateBinder& JavaScriptDelegateBinder,
+		FSimpleMulticastDelegate& UnbindDelegate) :
 		WebJavaScriptExecutor(JavaScriptExecutor),
 		WebJavaScriptDelegateBinder(JavaScriptDelegateBinder),
 		WebJavaScriptResultDelegate(NewObject<UAIAssistantWebJavaScriptResultDelegate>())
 	{
-		WebJavaScriptResultDelegate->Bind(*this);
+		WebJavaScriptResultDelegate->Bind(*this, UnbindDelegate);
 	}
 
 	FWebApi::~FWebApi()
diff --git a/Engine/Plugins/Experimental/AIAssistant/Source/AIAssistant/Private/AIAssistantWebApi.h b/Engine/Plugins/Experimental/AIAssistant/Source/AIAssistant/Private/AIAssistantWebApi.h
--- a/Engine/Plugins/Experimental/AIAssistant/Source/AIAssistant/Private/AIAssistantWebApi.h
+++ b/Engine/Plugins/Experimental/AIAssistant/Source/AIAssistant/Private/AIAssistantWebApi.h
@@ -4,6 +4,7 @@
 #include "Async/Future.h"
 #include "Containers/Array.h"
 #include "Containers/UnrealString.h"
+#include "Delegates/Delegate.h"
 #include "Misc/DateTime.h"
 #include "Misc/Optional.h"
 #include "Misc/Variant.h"
@@ -375,8 +376,10 @@ namespace UE::AIAssistant
 		friend class FWebApiAccessor;
 
 	public:
-		FWebApi(IWebJavaScriptExecutor& JavaScriptExecutor,
-				IWebJavaScriptDelegateBinder& JavaScriptDelegateBinder);
+		FWebApi(
+			IWebJavaScriptExecutor& JavaScriptExecutor,
+			IWebJavaScriptDelegateBinder& JavaScriptDelegateBinder,
+			FSimpleMulticastDelegate& UnbindDelegate);
 
 		// Prevent copy.
 		FWebApi(const FWebApi&) = delete;
diff --git a/Engine/Plugins/Experimental/AIAssistant/Source/AIAssistant/Private/AIAssistantWebApplication.cpp b/Engine/Plugins/Experimental/AIAssistant/Source/AIAssistant/Private/AIAssistantWebApplication.cpp
--- a/Engine/Plugins/Experimental/AIAssistant/Source/AIAssistant/Private/AIAssistantWebApplication.cpp
+++ b/Engine/Plugins/Experimental/AIAssistant/Source/AIAssistant/Private/AIAssistantWebApplication.cpp
@@ -413,11 +413,14 @@ namespace UE::AIAssistant
 
 	TFunction<TSharedPtr<FWebApi>()> FWebApplication::CreateWebApiFactory(
 		IWebJavaScriptExecutor& JavaScriptExecutor,
-		IWebJavaScriptDelegateBinder& JavaScriptDelegateBinder)
+		IWebJavaScriptDelegateBinder& JavaScriptDelegateBinder,
+		FSimpleMulticastDelegate& UnbindDelegate)
 	{
-		return [&JavaScriptExecutor, &JavaScriptDelegateBinder]() -> TSharedPtr<FWebApi>
+		return [&JavaScriptExecutor, &JavaScriptDelegateBinder, 
+				&UnbindDelegate]() -> TSharedPtr<FWebApi>
 			{
-				return MakeShared<FWebApi>(JavaScriptExecutor, JavaScriptDelegateBinder);
+				return MakeShared<FWebApi>(
+					JavaScriptExecutor, JavaScriptDelegateBinder, UnbindDelegate);
 			};
 	}
 
diff --git a/Engine/Plugins/Experimental/AIAssistant/Source/AIAssistant/Private/AIAssistantWebApplication.h b/Engine/Plugins/Experimental/AIAssistant/Source/AIAssistant/Private/AIAssistantWebApplication.h
--- a/Engine/Plugins/Experimental/AIAssistant/Source/AIAssistant/Private/AIAssistantWebApplication.h
+++ b/Engine/Plugins/Experimental/AIAssistant/Source/AIAssistant/Private/AIAssistantWebApplication.h
@@ -5,6 +5,7 @@
 #include "Async/Mutex.h"
 #include "Async/RecursiveMutex.h"
 #include "Containers/UnrealString.h"
+#include "Delegates/Delegate.h"
 #include "IWebBrowserWindow.h"
 #include "Misc/Optional.h"
 #include "Templates/Function.h"
@@ -187,6 +188,7 @@ namespace UE::AIAssistant
 		// Create a factory for web API instances.
 		static TFunction<TSharedPtr<FWebApi>()> CreateWebApiFactory(
 			IWebJavaScriptExecutor& JavaScriptExecutor,
-			IWebJavaScriptDelegateBinder& JavaScriptDelegateBinder);
+			IWebJavaScriptDelegateBinder& JavaScriptDelegateBinder,
+			FSimpleMulticastDelegate& UnbindDelegate);
 	};
 }
\ No newline at end of file
diff --git a/Engine/Plugins/Experimental/AIAssistant/Source/AIAssistant/Private/AIAssistantWebBrowser.cpp b/Engine/Plugins/Experimental/AIAssistant/Source/AIAssistant/Private/AIAssistantWebBrowser.cpp
--- a/Engine/Plugins/Experimental/AIAssistant/Source/AIAssistant/Private/AIAssistantWebBrowser.cpp
+++ b/Engine/Plugins/Experimental/AIAssistant/Source/AIAssistant/Private/AIAssistantWebBrowser.cpp
@@ -214,8 +214,11 @@ void SAIAssistantWebBrowser::Construct(const FArguments& InArgs, const TSharedRe
 void SAIAssistantWebBrowser::InitializeWebApplication()
 {
 	UE::TUniqueLock WebApplicationLock(WebApplicationMutex);
-	MaybeWebApplication = MakeShared<FWebApplication>(
-		FWebApplication::CreateWebApiFactory(*this, *this));
+	if (!CurrentConfig->IsExiting())
+	{
+		MaybeWebApplication = MakeShared<FWebApplication>(
+			FWebApplication::CreateWebApiFactory(*this, *this, CurrentConfig->OnPreExit()));
+	}
 }
 
 void SAIAssistantWebBrowser::EnsureWebApplication()
@@ -229,13 +232,14 @@ void SAIAssistantWebBrowser::WithWebApplication(
 	bool bCheckUsingWebApplicationIsCalled) const
 {
 	UE::TUniqueLock WebApplicationLock(WebApplicationMutex);
-	if (MaybeWebApplication.IsValid())
+	bool bIsExiting = CurrentConfig->IsExiting();
+	if (!bIsExiting && MaybeWebApplication.IsValid())
 	{
 		UsingWebApplication(MaybeWebApplication.ToSharedRef());
 	}
 	else
 	{
-		check(!bCheckUsingWebApplicationIsCalled);
+		check(!bCheckUsingWebApplicationIsCalled || bIsExiting);
 	}
 }
 
diff --git a/Engine/Plugins/Experimental/AIAssistant/Source/AIAssistant/Private/AIAssistantWebBrowser.h b/Engine/Plugins/Experimental/AIAssistant/Source/AIAssistant/Private/AIAssistantWebBrowser.h
--- a/Engine/Plugins/Experimental/AIAssistant/Source/AIAssistant/Private/AIAssistantWebBrowser.h
+++ b/Engine/Plugins/Experimental/AIAssistant/Source/AIAssistant/Private/AIAssistantWebBrowser.h
@@ -96,8 +96,9 @@ private:
 	void EnsureWebApplication();
 
 	// Get a reference to the current web application, if it's valid.
-	// If the reference is not valid assert / check on if bCheckUsingWebApplicationIsCalled
-	// is true.
+	// If the reference is not valid and the editor is not shutting down,
+	// assert / check on if bCheckUsingWebApplicationIsCalled is true.
+	// When the editor is shutting down UsingWebApplication will not be called.
 	void WithWebApplication(
 		TFunction<void(TSharedRef<UE::AIAssistant::FWebApplication>)>&& UsingWebApplication,
 		bool bCheckUsingWebApplicationIsCalled = true) const;
diff --git a/Engine/Plugins/Experimental/AIAssistant/Source/AIAssistant/Private/Tests/AIAssistantWebApiTest.cpp b/Engine/Plugins/Experimental/AIAssistant/Source/AIAssistant/Private/Tests/AIAssistantWebApiTest.cpp
--- a/Engine/Plugins/Experimental/AIAssistant/Source/AIAssistant/Private/Tests/AIAssistantWebApiTest.cpp
+++ b/Engine/Plugins/Experimental/AIAssistant/Source/AIAssistant/Private/Tests/AIAssistantWebApiTest.cpp
@@ -2,6 +2,7 @@
 
 #include "Containers/Array.h"
 #include "Containers/UnrealString.h"
+#include "Delegates/Delegate.h"
 #include "Internationalization/Regex.h"
 #include "Misc/AutomationTest.h"
 #include "Serialization/JsonReader.h"
@@ -24,10 +25,12 @@ namespace UE::AIAssistant
 {
 	struct FFakeWebApiContainer
 	{
-		FFakeWebApiContainer() : WebApi(WebJavaScriptExecutor, WebJavaScriptDelegateBinder) {}
+		FFakeWebApiContainer() :
+			WebApi(WebJavaScriptExecutor, WebJavaScriptDelegateBinder, UnbindDelegate) {}
 
 		FFakeWebJavaScriptExecutor WebJavaScriptExecutor;
 		FFakeWebJavaScriptDelegateBinder WebJavaScriptDelegateBinder;
+		FSimpleMulticastDelegate UnbindDelegate;
 		FFakeWebApi WebApi;
 
 		FFakeWebApi& operator*() { return WebApi; }
@@ -567,6 +570,22 @@ void FAIAssistantWebApiTest::Define()
 						TEXT(R"json(null)json"), Result, *Conversation.ToJson(false), false);
 				});
 		});
+
+	Describe(TEXT("UnbindDelegate"), [this]
+		{
+			It(TEXT("Should cancel callbacks on broadcast"), [this]
+				{
+					FFakeWebApiContainer WebApi;					
+					auto Result = WebApi->CreateConversation();
+					WebApi.UnbindDelegate.Broadcast();
+					auto NothingOrError = Result.Consume();
+					if (!TestTrue(TEXT("HasError"), NothingOrError.HasError())) return;
+					(void)TestEqual(
+						TEXT("Canceled"), NothingOrError.GetError(),
+						UAIAssistantWebJavaScriptResultDelegate::CanceledError);
+				});
+		}
+	);
 }
 
 
diff --git a/Engine/Plugins/Experimental/AIAssistant/Source/AIAssistant/Private/Tests/AIAssistantWebApplicationTest.cpp b/Engine/Plugins/Experimental/AIAssistant/Source/AIAssistant/Private/Tests/AIAssistantWebApplicationTest.cpp
--- a/Engine/Plugins/Experimental/AIAssistant/Source/AIAssistant/Private/Tests/AIAssistantWebApplicationTest.cpp
+++ b/Engine/Plugins/Experimental/AIAssistant/Source/AIAssistant/Private/Tests/AIAssistantWebApplicationTest.cpp
@@ -2,6 +2,7 @@
 
 #include "Containers/Array.h"
 #include "Containers/UnrealString.h"
+#include "Delegates/Delegate.h"
 #include "Internationalization/Culture.h"
 #include "IWebBrowserWindow.h"
 #include "Misc/AutomationTest.h"
@@ -18,6 +19,7 @@
 #include "AIAssistantToolResponse.h"
 #include "AIAssistantWebApi.h"
 #include "AIAssistantWebApplication.h"
+#include "AIAssistantWebJavaScriptResultDelegate.h"
 #include "Tests/AIAssistantFakeToolsetHandler.h"
 #include "Tests/AIAssistantFakeWebJavaScriptDelegateBinder.h"
 #include "Tests/AIAssistantFakeWebJavaScriptExecutor.h"
@@ -61,12 +63,14 @@ namespace UE::AIAssistant
 	struct FFakeWebApiWithFakeJavaScriptEnvironment
 	{
 		FFakeWebApiWithFakeJavaScriptEnvironment() :
-			WebApi(MakeShared<FFakeWebApi>(JavaScriptExecutor, JavaScriptDelegateBinder))
+			WebApi(MakeShared<FFakeWebApi>(
+				JavaScriptExecutor, JavaScriptDelegateBinder, UnbindDelegate))
 		{
 		}
 
 		FFakeWebJavaScriptExecutor JavaScriptExecutor;
 		FFakeWebJavaScriptDelegateBinder JavaScriptDelegateBinder;
+		FSimpleMulticastDelegate UnbindDelegate;
 		TSharedPtr<FFakeWebApi> WebApi;
 	};
 
@@ -292,34 +296,74 @@ namespace UE::AIAssistant
 }
 
 
-IMPLEMENT_SIMPLE_AUTOMATION_TEST(
+BEGIN_DEFINE_SPEC(
 	FAIAssistantWebApplicationTestCreateWebApiFactory,
 	"AI.Assistant.WebApplication.CreateWebApiFactory",
-	AIAssistantTest::Flags);
+	AIAssistantTest::Flags)
+END_DEFINE_SPEC(FAIAssistantWebApplicationTestCreateWebApiFactory)
 
-bool FAIAssistantWebApplicationTestCreateWebApiFactory::RunTest(const FString& UnusedParameters)
+void FAIAssistantWebApplicationTestCreateWebApiFactory::Define()
 {
-	FFakeWebJavaScriptExecutor JavaScriptExecutor;
-	FFakeWebJavaScriptDelegateBinder JavaScriptDelegateBinder;
-	auto Factory = FWebApplication::CreateWebApiFactory(
-		JavaScriptExecutor, JavaScriptDelegateBinder);
-	if (!TestTrue(TEXT("FactoryIsSet"), Factory.IsSet())) return false;
-
-	auto WebApi = Factory();
-	if (!TestTrue(TEXT("FactoryCreatedWebApi"), WebApi.IsValid())) return false;
-
-	// Ensure the web API is bound to the JavaScript environment.
-	(void)TestEqual(
-		TEXT("BoundApiToJavaScriptEnvironment"),
-		JavaScriptDelegateBinder.BoundObjects.Num(), 1);
-	// Make sure the web API is using the supplied executor.
-	WebApi->UpdateGlobalLocale(TEXT("us-en"));
-	(void)TestNotEqual(
-		TEXT("ExecutedFunction"), JavaScriptExecutor.ExecutedJavaScriptText.Num(), 0);
-
-	auto AnotherWebApi = Factory();
-	if (!TestNotEqual(TEXT("FactoryCreatesNewWebApis"), WebApi, AnotherWebApi)) return false;
-	return true;
+	It(TEXT("should create a WebApi"), [this]
+		{
+			FFakeWebJavaScriptExecutor JavaScriptExecutor;
+			FFakeWebJavaScriptDelegateBinder JavaScriptDelegateBinder;
+			FSimpleMulticastDelegate UnbindDelegate;
+			auto Factory = FWebApplication::CreateWebApiFactory(
+				JavaScriptExecutor, JavaScriptDelegateBinder, UnbindDelegate);
+			if (!TestTrue(TEXT("FactoryIsSet"), Factory.IsSet())) return;
+
+			auto WebApi = Factory();
+			if (!TestTrue(TEXT("FactoryCreatedWebApi"), WebApi.IsValid())) return;
+
+			// Ensure the web API is bound to the JavaScript environment.
+			(void)TestEqual(
+				TEXT("BoundApiToJavaScriptEnvironment"),
+				JavaScriptDelegateBinder.BoundObjects.Num(), 1);
+			// Make sure the web API is using the supplied executor.
+			WebApi->UpdateGlobalLocale(TEXT("us-en"));
+			(void)TestNotEqual(
+				TEXT("ExecutedFunction"), JavaScriptExecutor.ExecutedJavaScriptText.Num(), 0);
+		});
+
+	It(TEXT("should create a unique WebApi"), [this]
+		{
+			FFakeWebJavaScriptExecutor JavaScriptExecutor;
+			FFakeWebJavaScriptDelegateBinder JavaScriptDelegateBinder;
+			FSimpleMulticastDelegate UnbindDelegate;
+			auto Factory = FWebApplication::CreateWebApiFactory(
+				JavaScriptExecutor, JavaScriptDelegateBinder, UnbindDelegate);
+			if (!TestTrue(TEXT("FactoryIsSet"), Factory.IsSet())) return;
+
+			auto WebApi = Factory();
+			(void)TestTrue(TEXT("FactoryCreatedWebApi"), WebApi.IsValid());
+
+			auto AnotherWebApi = Factory();
+			(void)TestTrue(TEXT("FactoryCreatedNewWebApi"), AnotherWebApi.IsValid());
+			(void)TestNotEqual(TEXT("FactoryCreatesNewWebApis"), WebApi, AnotherWebApi);
+		});
+
+	It(TEXT("should create a WebApi with callbacks that can be canceled"), [this]
+		{
+			FFakeWebJavaScriptExecutor JavaScriptExecutor;
+			FFakeWebJavaScriptDelegateBinder JavaScriptDelegateBinder;
+			FSimpleMulticastDelegate UnbindDelegate;
+			auto Factory = FWebApplication::CreateWebApiFactory(
+				JavaScriptExecutor, JavaScriptDelegateBinder, UnbindDelegate);
+			if (!TestTrue(TEXT("FactoryIsSet"), Factory.IsSet())) return;
+
+			auto WebApi = Factory();
+			if (!TestTrue(TEXT("FactoryCreatedWebApi"), WebApi.IsValid())) return;
+
+			auto Result = WebApi->CreateConversation();
+			UnbindDelegate.Broadcast();
+
+			auto NothingOrError = Result.Consume();
+			if (!TestTrue(TEXT("HasError"), NothingOrError.HasError())) return;
+			(void)TestEqual(
+				TEXT("Canceled"), NothingOrError.GetError(),
+				UAIAssistantWebJavaScriptResultDelegate::CanceledError);
+		});
 }
 
 IMPLEMENT_SIMPLE_AUTOMATION_TEST(
