AI Assistant, Tool Calling

- FObjectFunctionToolCall gets and stores the UFunction's JSON schema on construction. This is used during execution for -
(1) setting missing non-required input params to defaults from the schema.
(2) checking that all required input params are provided. TODO - This ultimately should become full function/schema validation. But for now, this is a necessary step for (1).

#Synced-CL 49419081

diff --git a/Engine/Plugins/Experimental/AIAssistant/Source/AIAssistant/Private/AIAssistantToolCall.cpp b/Engine/Plugins/Experimental/AIAssistant/Source/AIAssistant/Private/AIAssistantToolCall.cpp
--- a/Engine/Plugins/Experimental/AIAssistant/Source/AIAssistant/Private/AIAssistantToolCall.cpp
+++ b/Engine/Plugins/Experimental/AIAssistant/Source/AIAssistant/Private/AIAssistantToolCall.cpp
@@ -25,32 +25,42 @@ using namespace UE::AIAssistant;
 
 namespace
 {
-	TSharedPtr<FJsonObject> JsonStringToJsonObject(const FString& JsonString)
+	TSharedRef<FJsonObject> JsonStringToValidJsonObject(const FString& JsonString)
 	{
 		if (!JsonString.IsEmpty())
 		{
-			const TSharedRef<TJsonReader<>> JsonReader = TJsonReaderFactory<>::Create(JsonString);
+			const TSharedRef<TJsonReader<>> JsonReader = 
+				TJsonReaderFactory<>::Create(JsonString);
 			TSharedPtr<FJsonObject> JsonObject;
 			if (FJsonSerializer::Deserialize(JsonReader, JsonObject))
 			{
-				return JsonObject;
+				return JsonObject.ToSharedRef();
 			}
 		}
-		
-		return nullptr;			
+
+		return MakeShared<FJsonObject>();			
 	}
 	
-	TSharedPtr<FJsonObject> FunctionInputParamsJsonToJsonObject(
-		const FObjectFunctionToolCall::FFunctionInputParamsJson& FunctionInputParamsJson)
+	TSharedRef<FJsonObject> FunctionInputParamsJsonToValidJsonObject(
+		const TOptional<FObjectFunctionToolCall::FFunctionInputParamsJson>& FunctionInputParamsJson)
 	{
-		if (FunctionInputParamsJson.IsType<TSharedPtr<FJsonObject>>())
-		{
-			return FunctionInputParamsJson.Get<TSharedPtr<FJsonObject>>();
-		}
-		else
+		if (FunctionInputParamsJson.IsSet())
 		{
-			return JsonStringToJsonObject(FunctionInputParamsJson.Get<FString>());
+			if (FunctionInputParamsJson->IsType<TSharedPtr<FJsonObject>>())
+			{
+				// Return a copy of the JSON object, since it may be altered, and we 
+				// don't want to alter the original.
+				const TSharedPtr<FJsonObject> OriginalJsonObject = 
+					FunctionInputParamsJson->Get<TSharedPtr<FJsonObject>>();
+				return OriginalJsonObject ? MakeShared<FJsonObject>(*OriginalJsonObject) : MakeShared<FJsonObject>();
+			}
+			else
+			{
+				return JsonStringToValidJsonObject(FunctionInputParamsJson->Get<FString>());
+			}
 		}
+		
+		return MakeShared<FJsonObject>();
 	}
 	
 	UWorld* GetWorld()
@@ -59,7 +69,8 @@ namespace
 		{
 			for (const FWorldContext& WorldContext : GEngine->GetWorldContexts())
 			{
-				if (WorldContext.WorldType == EWorldType::Game || WorldContext.WorldType == EWorldType::PIE)
+				if (WorldContext.WorldType == EWorldType::Game ||
+					WorldContext.WorldType == EWorldType::PIE)
 				{
 					return WorldContext.World();
 				}
@@ -69,7 +80,8 @@ namespace
 		return nullptr;
 	}
 	
-	void MaybeSetWorldContextPropertyInFrameMemory(const UFunction* Function, uint8* FrameMemory)
+	void MaybeSetWorldContextPropertyInFrameMemory(const UFunction* Function, 
+		uint8* FrameMemory)
 	{
 		const FName WorldContextKey(TEXT("WorldContext"));
 		const FName WorldContextParamName = 
@@ -77,16 +89,112 @@ namespace
 			FName(*Function->GetMetaData(WorldContextKey)) :
 			FName(TEXT("WorldContextObject"));
 		
-		if (FProperty* WorldContextProperty = Function->FindPropertyByName(WorldContextParamName))
+		FProperty* WorldContextProperty = Function->FindPropertyByName(WorldContextParamName);
+		if (WorldContextProperty)
 		{
 			UWorld* World = GetWorld();
-			FObjectPropertyBase* ObjectProperty = CastField<FObjectPropertyBase>(WorldContextProperty);
+			FObjectPropertyBase* ObjectProperty = 
+				CastField<FObjectPropertyBase>(WorldContextProperty);
 			if (ObjectProperty && World)
 			{
 				ObjectProperty->SetObjectPropertyValue_InContainer(FrameMemory, World);
 			}
 		}
 	}
+
+	/**
+	  * Checks whether all the input parameters required by the function schema JSON are provided
+	  * in the function input params JSON.
+	  * @param FunctionInputParamsJsonObject The JSON option containing the function call params.
+	  * @param FunctionSchemaInputSchemaJsonObject The 'inputSchema' sub-JSON in the main function 
+	  *		schema JSON.
+	  *	@param OutErrorText Output error text if a required param is missing, empty otherwise.
+	  *	@return Whether required params are all supplied.
+	  */
+	bool AreAllRequiredFunctionInputParamsProvided(
+		const TSharedRef<FJsonObject>& FunctionInputParamsJsonObject,
+		const TSharedRef<FJsonObject>& FunctionSchemaInputSchemaJsonObject,
+		FText& OutErrorText)
+	{
+		// For each required param name, make sure we have a field for it in the function input params.
+		{
+			TArray<FString> RequiredParamNames;
+			FunctionSchemaInputSchemaJsonObject->TryGetStringArrayField(
+				TEXT("required"), RequiredParamNames);
+
+			if (!RequiredParamNames.IsEmpty() && FunctionInputParamsJsonObject->Values.IsEmpty())
+			{
+				// Caller error. Missing but expected params in incoming function 
+				// input params JSON. 
+				static const FText ErrorText = LOCTEXT(
+					"AIAssistantToolCallError_FunctionInputParamsJsonIsMissingRequiredParams",	
+					"Function input params are required by the function input schema JSON, "
+					"but is incoming function input params JSON is empty.");
+				OutErrorText = ErrorText;
+				return false;
+			}
+    		
+			for (const FString& RequiredParamName : RequiredParamNames)
+			{
+				if (!FunctionInputParamsJsonObject->HasField(*RequiredParamName))
+				{
+					// Caller error. A missing but expected param in incoming function 
+					// input params JSON. 
+					static const FText ErrorTextTemplate = LOCTEXT(
+						"AIAssistantToolCallError_FunctionInputParamsJsonIsMissingRequiredParam",	
+						"Function input param '{0}' is required by the function input schema JSON, "
+						"but is missing from the incoming function input params JSON.");
+					OutErrorText = FText::Format(
+						ErrorTextTemplate, FText::FromString(RequiredParamName));
+					return false;
+				}
+			}
+		}
+	
+		return true;
+	}
+
+	/**
+	 * For each input schema property, if we don't have a field for it in the 
+	 * function input params, then add it to the function input params, and 
+	 * give it the same value as the default value of the corresponding function
+	 * input schema property.
+	 * (Note - Checked that we have all the required params first, so we know those 
+	 * will all be in the function input params JSON. Then this ends up only handling
+	 * params that are both non-required and missing from the function input param 
+	 * JSON. To apply default values, we need a 'default' value to exist for each 
+	 * corresponding param object in the function input schema.)
+	  * @param FunctionInputParamsJsonObject The JSON option containing the function call params.
+	  * @param FunctionSchemaInputSchemaJsonObject The 'inputSchema' sub-JSON in the main function 
+	  *		schema JSON.
+	 */
+	void ProvideDefaultsForMissingFunctionInputParams(
+		const TSharedRef<FJsonObject>& FunctionInputParamsJsonObject,
+		const TSharedRef<FJsonObject>& FunctionSchemaInputSchemaJsonObject)
+	{
+		if (const TSharedPtr<FJsonObject>* InputPropertiesJsonObject;
+			FunctionSchemaInputSchemaJsonObject->TryGetObjectField(
+				TEXT("properties"), InputPropertiesJsonObject))
+		{
+			for (const TPair<FString, TSharedPtr<FJsonValue>>& Pair : 
+				(*InputPropertiesJsonObject)->Values)
+			{
+				const FString& InputParamPropertyName = Pair.Key;
+
+				const TSharedPtr<FJsonObject>& InputPropertyJsonObject = Pair.Value->AsObject();
+				check(InputPropertyJsonObject)
+    			
+				if (!FunctionInputParamsJsonObject->HasField(*InputParamPropertyName))
+				{
+					const TSharedPtr<FJsonValue> InputPropertyDefaultJsonValue = 
+						InputPropertyJsonObject->TryGetField(TEXT("default"));
+					check(InputPropertyDefaultJsonValue); // ..this really needs o exist
+					FunctionInputParamsJsonObject->SetField(*InputParamPropertyName, 
+						InputPropertyDefaultJsonValue);
+				}
+			}
+		}
+	}
 }
 
 //
@@ -98,24 +206,50 @@ namespace
 	FModuleManager::Get().LoadModuleChecked(TEXT("JsonUtilitiesEditor"));
 }
 
-FObjectFunctionToolCall::FObjectFunctionToolCall(TNotNull<UObject*> InInstanceObject, TNotNull<UFunction*> InFunction) :
+FObjectFunctionToolCall::FObjectFunctionToolCall(
+		TNotNull<UObject*> InInstanceObject, TNotNull<UFunction*> InFunction) :
 	InstanceObject(InInstanceObject), 
 	Function(InFunction)
 {
+	FunctionSchemaJsonObject = FJsonSchemaGenerator::UStructToJsonSchemaObject(Function);
+	check(FunctionSchemaJsonObject);
+}
+
+/*static*/ FObjectFunctionToolCall FObjectFunctionToolCall::ConstructWithNamedClassFunction(
+	TNotNull<UObject*> InInstanceObject, const FName& InFunctionName)
+{
+	check(!InFunctionName.IsNone());
+		
+	UFunction* Function = InInstanceObject->GetClass()->FindFunctionByName(InFunctionName);
+	check(Function);
+
+	return FObjectFunctionToolCall(InInstanceObject, Function);
 }
 
-FObjectFunctionToolCall::FToolCallResult FObjectFunctionToolCall::Execute(const FFunctionInputParamsJson& FunctionInputParamsJson) const
+FObjectFunctionToolCall::FToolCallResult FObjectFunctionToolCall::Execute(
+	const TOptional<FFunctionInputParamsJson>& FunctionInputParamsJson) const
 {
-	// Get incoming function params JSON object.
-	const TSharedPtr<FJsonObject> FunctionInputParamsJsonObject = 
-		FunctionInputParamsJsonToJsonObject(FunctionInputParamsJson);
+	// Get function input params JSON as a JSON object.
+	TSharedPtr<FJsonObject> FunctionInputParamsJsonObject;
+	if (const TValueOrError<const TSharedRef<FJsonObject>, const FText> JsonObjectOrError = 
+			BuildValidFunctionInputParamsJsonObject(FunctionInputParamsJson);
+		JsonObjectOrError.HasError())
+	{
+		return MakeError(JsonObjectOrError.GetError());
+	}
+	else
+	{
+		FunctionInputParamsJsonObject = JsonObjectOrError.GetValue();
+	}
 	
 	// Allocate the frame memory.
 	uint8* FrameMemory;
 	{
-		// Invoke must use properties size, because it's allocating locals for e.g. script functions.
+		// Invoke must use properties size, because it's allocating locals, 
+		// i.e. for script functions.
 		const int32 FrameSize = Function->PropertiesSize; 
-		FrameMemory = static_cast<uint8*>(FMemory_Alloca_Aligned(FrameSize, Function->GetMinAlignment()));
+		FrameMemory = static_cast<uint8*>(FMemory_Alloca_Aligned(FrameSize, 
+			Function->GetMinAlignment()));
 		FMemory::Memzero(FrameMemory, FrameSize);
 	}
 	
@@ -130,9 +264,10 @@ FObjectFunctionToolCall::FToolCallResult FObjectFunctionToolCall::Execute(const
 	}
 
 	// When we exit, we need to manually destroy each property in the frame memory.
-	ON_SCOPE_EXIT
+	ON_SCOPE_EXIT -> void
 	{
-		// Invoke will destroy locals, but caller owns parameters. This should also destroy the result parameter.
+		// Invoke will destroy locals, but caller owns parameters. This should also 
+		// destroy the result parameter.
 		for (TFieldIterator<FProperty> It(Function); It && It->HasAnyPropertyFlags(CPF_Parm); ++It)
 		{
 			It->DestroyValue_InContainer(FrameMemory);
@@ -140,16 +275,15 @@ FObjectFunctionToolCall::FToolCallResult FObjectFunctionToolCall::Execute(const
 	};
 
 	// Copy the function input params memory into the function's frame memory.
-	if (FunctionInputParamsJsonObject)
+	if (!FJsonObjectConverter::JsonObjectToUStruct(
+			FunctionInputParamsJsonObject.ToSharedRef(), Function, FrameMemory))
 	{
-		const bool bDidConvertObjectToUStruct = FJsonObjectConverter::JsonObjectToUStruct(FunctionInputParamsJsonObject.ToSharedRef(), Function, 
-			FrameMemory);
-		check(bDidConvertObjectToUStruct);
-		if (!bDidConvertObjectToUStruct)
-		{
-			static const FText ErrorString = LOCTEXT("JsonObjectToUStructFailed", "Could not convert input params JSON object to a UStruct.");
-			return MakeError(ErrorString);
-		}
+		// Caller error. Ill-formed incoming function input params JSON. 
+		static const FText ErrorText = LOCTEXT(
+			"AIAssistantToolCallError_JsonObjectToUStructFailed", 
+			"Could not convert incoming function input params JSON to a UStruct. "
+			"Incoming function input params JSON may be incorrect or contain unsupported types.");
+		return MakeError(ErrorText);
 	}
 	
 	// In case this is a Blueprint function that has a world context object parameter, 
@@ -157,17 +291,30 @@ FObjectFunctionToolCall::FToolCallResult FObjectFunctionToolCall::Execute(const
 	MaybeSetWorldContextPropertyInFrameMemory(Function, FrameMemory);
 
 	// Invoke function.
-	FFrame StackFrame(InstanceObject, Function, FrameMemory, nullptr, Function->ChildProperties);
+	FFrame StackFrame(InstanceObject, Function, FrameMemory, nullptr, 
+		Function->ChildProperties);
 	const bool bHasReturnParam = (Function->ReturnValueOffset != MAX_uint16);
-	uint8* ReturnValueAddress = (bHasReturnParam ? (FrameMemory + Function->ReturnValueOffset) : nullptr);
+	uint8* ReturnValueAddress = (bHasReturnParam ? 
+		(FrameMemory + Function->ReturnValueOffset) : nullptr);
 	Function->Invoke(InstanceObject, StackFrame, ReturnValueAddress);
 	
 	// Get the return value, if we have one.
 	TSharedPtr<FJsonObject> ReturnJsonObject; 
 	if (FProperty* ReturnProperty = Function->GetReturnProperty())
 	{
-		if (const TSharedPtr<FJsonValue> ReturnJsonValue = FJsonObjectConverter::UPropertyToJsonValue(ReturnProperty, ReturnValueAddress);
-			ReturnJsonValue->Type == EJson::Object)
+		const TSharedPtr<FJsonValue> ReturnJsonValue = 
+			FJsonObjectConverter::UPropertyToJsonValue(ReturnProperty, ReturnValueAddress);
+		if (!ReturnJsonValue)
+		{
+			// Caller error. Unsupported return value type. 
+			static const FText ErrorText = LOCTEXT(
+				"AIAssistantToolCallError_ReturnValueToJsonValueFailed", 
+				"Could not convert function return value to JSON value. "
+				"Function return value type may be unsupported.");
+			return MakeError(ErrorText);
+		}
+		
+		if (ReturnJsonValue->Type == EJson::Object)
 		{
 			// If value is an object, return it.
 			ReturnJsonObject = ReturnJsonValue->AsObject();
@@ -183,6 +330,44 @@ FObjectFunctionToolCall::FToolCallResult FObjectFunctionToolCall::Execute(const
 	return MakeValue(ReturnJsonObject);
 }
 
+TValueOrError<const TSharedRef<FJsonObject>, const FText> 
+FObjectFunctionToolCall::BuildValidFunctionInputParamsJsonObject(
+	const TOptional<FFunctionInputParamsJson>& FunctionInputParamsJson) const
+{
+	check(FunctionSchemaJsonObject);
+
+	// Get incoming function params as JSON object.
+	const TSharedRef<FJsonObject> FunctionInputParamsJsonObject = 
+		FunctionInputParamsJsonToValidJsonObject(FunctionInputParamsJson);
+	
+	// With function schema JSON -
+	// (1) Validate the function input JSON against the function schema's input schema JSON.
+	// (2) Set any missing params in the function input params JSON to the defaults from 
+	//		function schema's input schema JSON.
+    if (const TSharedPtr<FJsonObject>* FunctionSchemaInputSchemaJsonObjectPtr;
+    	FunctionSchemaJsonObject->TryGetObjectField(
+    		TEXT("inputSchema"), FunctionSchemaInputSchemaJsonObjectPtr))
+    {
+		// TODO - We should ultimately be validating the function input params JSON against the 
+		// function schema's input schema JSON instead of this. This isn't validating type
+		// compatibility, only checking existence.
+		if (FText ErrorText;
+			!AreAllRequiredFunctionInputParamsProvided(
+				FunctionInputParamsJsonObject, 
+				FunctionSchemaInputSchemaJsonObjectPtr->ToSharedRef(), 
+				ErrorText))
+		{
+			return MakeError(ErrorText);
+		}
+
+		ProvideDefaultsForMissingFunctionInputParams(
+			FunctionInputParamsJsonObject, 
+			FunctionSchemaInputSchemaJsonObjectPtr->ToSharedRef()); 
+	}
+	
+	return MakeValue(FunctionInputParamsJsonObject);
+}
+
 PRAGMA_ENABLE_DEPRECATION_WARNINGS
 
 #undef LOCTEXT_NAMESPACE
diff --git a/Engine/Plugins/Experimental/AIAssistant/Source/AIAssistant/Private/AIAssistantToolCall.h b/Engine/Plugins/Experimental/AIAssistant/Source/AIAssistant/Private/AIAssistantToolCall.h
--- a/Engine/Plugins/Experimental/AIAssistant/Source/AIAssistant/Private/AIAssistantToolCall.h
+++ b/Engine/Plugins/Experimental/AIAssistant/Source/AIAssistant/Private/AIAssistantToolCall.h
@@ -20,24 +20,20 @@ namespace UE::AIAssistant
 	public:
 
 		/**
+		 * Function input params.
 		 * Holds either a JSON object or a JSON string.
-		 * Usage -
-		 *		FFunctionInputParamsJson FunctionInputParamsJson;
-		 *		...
-		 *		FunctionInputParamsJson.Set<FString>(SomeJsonString);
-		 *		or
-		 *		FunctionInputParamsJson.Set<TSharedPtr<FJsonObject>>(SomeJsonSharedObjectPtr);
 		 */
 		using FFunctionInputParamsJson = TVariant<TSharedPtr<FJsonObject>, FString>;
 
 		/**
-		 * Tool call result.
-		 * Either JSON object or text describing the error.
+		 * Tool call output result.
+		 * Either JSON object or text describing an error.
 		 */
 		using FToolCallResult = TValueOrError<const TSharedPtr<FJsonObject>, const FText>;
 		
 		/**
 		 * Initializes use of this class.
+		 * Makes sure necessary modules are loaded.
 		 */
 		static void Initialize();
 		
@@ -46,35 +42,44 @@ namespace UE::AIAssistant
 		 * @param InInstanceObject UObject to call the function on.
 		 * @param InFunction UFunction to call on the UObject.
 		 */
-		FObjectFunctionToolCall(TNotNull<UObject*> InInstanceObject, TNotNull<UFunction*> InFunction);
+		FObjectFunctionToolCall(
+			TNotNull<UObject*> InInstanceObject, TNotNull<UFunction*> InFunction);
 		
 		/**
-		 * Constructs with a UFunction on a UObject instance, but specialized for UObjects that are UClasses where the function can be found by name.
+		 * Constructs with a UFunction on a UObject instance, but using the UObject's UClass 
+		 * where the function can be found by name.
 		 * @param InInstanceObject UObject to call the function on.
 		 * @param InFunctionName UFunction name to find and call on the UObject.
+		 * @return A function call object.
 		 */
-		static TOptional<FObjectFunctionToolCall> FromNamedClassFunction(TNotNull<UObject*> InInstanceObject, const FName& InFunctionName)
-		{
-			check(!InFunctionName.IsNone());
-		
-			UFunction* Function = InInstanceObject->GetClass()->FindFunctionByName(InFunctionName);
-			check(Function);
-
-			return TOptional(FObjectFunctionToolCall(InInstanceObject, Function));
-		}
+		static FObjectFunctionToolCall ConstructWithNamedClassFunction(
+			TNotNull<UObject*> InInstanceObject, const FName& InFunctionName);
 		
 		/**
-		 * Given a JSON containing input parameters for a tool call, execute a function with those input parameters.
-		 * @param FunctionInputParamsJson The JSON object or string containing the function parameters.
+		 * Given a JSON containing input parameters for a tool call, execute the function with those 
+		 * input parameters.
+		 * @param FunctionInputParamsJson The optional JSON object or string containing the function 
+		 *		parameters. If the function this class was initialized with has required params, then
+		 *		this must be supplied with JSON that provides those required params. Otherwise, an
+		 *		error will result. If the function this class was initialized with has no required 
+		 *		params, then this does not need to be supplied.
 		 * @return Tool call result. (See above). When tool result does not have an error -
- 		 *		Returned objects will either be (1) the object itself, or (2) a value named 'returnValue' inside an object.
-		 *		May be null if (1) function does not return a value, or (2) if function does return a value and that return value was null.
+ 		 *		Returned objects will either be (1) the object itself, or (2) a value named 
+ 		 *		'returnValue' inside an object. And this may be null if (1) function does not 
+ 		 *		return a value, or (2) if function does return a value and that return value 
+ 		 *		is null.
 		 */
-		FToolCallResult Execute(const FFunctionInputParamsJson& FunctionInputParamsJson) const;
+		FToolCallResult Execute(const TOptional<FFunctionInputParamsJson>& FunctionInputParamsJson = 
+			TOptional<FFunctionInputParamsJson>()) const;
 		
 	private:
+
+		TValueOrError<const TSharedRef<FJsonObject>, const FText> 
+		BuildValidFunctionInputParamsJsonObject(
+			const TOptional<FFunctionInputParamsJson>& FunctionInputParamsJson) const;
 		
 		TNotNull<UObject*> InstanceObject;
 		TNotNull<UFunction*> Function;
+		TSharedPtr<FJsonObject> FunctionSchemaJsonObject;
 	};
 }
diff --git a/Engine/Plugins/Experimental/AIAssistant/Source/AIAssistant/Private/Tests/AIAssistantToolCallTest.cpp b/Engine/Plugins/Experimental/AIAssistant/Source/AIAssistant/Private/Tests/AIAssistantToolCallTest.cpp
--- a/Engine/Plugins/Experimental/AIAssistant/Source/AIAssistant/Private/Tests/AIAssistantToolCallTest.cpp
+++ b/Engine/Plugins/Experimental/AIAssistant/Source/AIAssistant/Private/Tests/AIAssistantToolCallTest.cpp
@@ -11,13 +11,10 @@
 
 #include "AIAssistantToolCall.h"
 
-
 using namespace UE::AIAssistant;
-	
 
 PRAGMA_DISABLE_DEPRECATION_WARNINGS
 
-
 #define PRINT_JSONS_TO_LOG 1
 
 static void MaybePrintJsonToLog(const TCHAR* Name, const TSharedRef<FJsonObject>& JsonObject)
@@ -32,74 +29,174 @@ static void MaybePrintJsonToLog(const TCHAR* Name, const TSharedRef<FJsonObject>
 #endif
 }
 
-
 #if WITH_DEV_AUTOMATION_TESTS
 
+// Test function with required param.
 
 IMPLEMENT_SIMPLE_AUTOMATION_TEST(
-	FAIAssistantToolCallExecuteToolCallFunction,
-	"AI.Assistant.ToolCall.ExecuteToolCallFunction",
+	FAIAssistantToolCallExecuteToolCallFunctionWithRequiredParam,
+	"AI.Assistant.ToolCall.ExecuteToolCallFunctionWithRequiredParam",
 	AIAssistantTest::Flags);
 
-bool FAIAssistantToolCallExecuteToolCallFunction::RunTest(const FString& Parameters)
+bool FAIAssistantToolCallExecuteToolCallFunctionWithRequiredParam::RunTest(
+	const FString& Parameters)
 {
-	static const FName FunctionName = "TestFunc";
+	static const FName FunctionName = "TestFuncWithRequiredParam";
+	static constexpr float ClassFloatValue = 20.0f;
 	static constexpr float FunctionInputFloatValue = 3.0f;
+	static const FString FunctionInputParamsJsonString = 
+		FString::Printf(TEXT("{\n\"InFloat\": %f\n}"), FunctionInputFloatValue);
+
+	// Create an object and provide it specific internal state. 
+	// (For this test - Using TStrongObjectPtr to prevent GC.)
+	
+	const TStrongObjectPtr InstanceObject = 
+		TStrongObjectPtr(NewObject<UToolCallFakeClass>());
+
+	InstanceObject->Float = ClassFloatValue;
+	
+	// Execute.
+	
+	const FObjectFunctionToolCall ToolCall = 
+		FObjectFunctionToolCall::ConstructWithNamedClassFunction(InstanceObject.Get(), FunctionName);
+	
+	FObjectFunctionToolCall::FFunctionInputParamsJson FunctionInputParamsJson;
+	FunctionInputParamsJson.Set<FString>(FunctionInputParamsJsonString);
+	const FObjectFunctionToolCall::FToolCallResult ToolCallResult = 
+		ToolCall.Execute(FunctionInputParamsJson);
+	
+	if (!TestTrue("Execute tool call result does not have error.", 
+			!ToolCallResult.HasError()))
+	{
+		UE_LOG(LogAIAssistant, Error, TEXT("Tool call error - %s"), 
+			*ToolCallResult.GetError().ToString());
+		return false;
+	}
+	if (!TestTrue("Execute tool call result has value.", 
+			ToolCallResult.HasValue()))
+	{
+		return false;
+	}
+
+	const TSharedPtr<FJsonObject> JsonObject = ToolCallResult.GetValue();
+
+	if (!TestTrue("Execute tool call result has valid value.", 
+			JsonObject.IsValid()))
+	{
+		return false;
+	}
+
+	MaybePrintJsonToLog(TEXT("ToolCallWithRequiredParamReturnedJsonObject"), 
+		JsonObject.ToSharedRef());
+	
+	if (!TestTrue("Return object is not empty (has values.)", 
+			JsonObject->Values.Num() > 0))
+	{
+		return false;
+	}	
+	if (!TestTrue(TEXT("Return object has float returnValue."), 
+			JsonObject->HasTypedField(TEXT("returnValue"), EJson::Number)))
+	{
+		return false;
+	}
+	
+	// Check.
+	
+	const float ReturnValueFloat = 
+		JsonObject->GetNumberField(TEXT("returnValue"));
+	const float ReturnValueFloatExpected = 
+		InstanceObject->TestFuncWithRequiredParam(FunctionInputFloatValue);
+	const FString TestMessage = FString::Printf(TEXT("Return value %f is expected value %f."), 
+		ReturnValueFloat, ReturnValueFloatExpected);
+	if (!TestTrue(TestMessage, 
+			ReturnValueFloat == ReturnValueFloatExpected))
+	{
+		return false;
+	}
+	
+	// Success.
+	
+	return true;
+}
+
+// Test function with optional param.
+
+IMPLEMENT_SIMPLE_AUTOMATION_TEST(
+	FAIAssistantToolCallExecuteToolCallFunctionWithOptionalParam,
+	"AI.Assistant.ToolCall.ExecuteToolCallFunctionOptionalParam",
+	AIAssistantTest::Flags);
+
+bool FAIAssistantToolCallExecuteToolCallFunctionWithOptionalParam::RunTest(
+	const FString& Parameters)
+{
+	static const FName FunctionName = "TestFuncWithOptionalParam";
 	static constexpr float ClassFloatValue = 20.0f;
-	static const FString FunctionInputParamsJsonString = FString::Printf(TEXT("{\n\"InFloat\": %f\n}"), FunctionInputFloatValue);
 
 	// Create an object and provide it specific internal state. 
 	// (For this test - Using TStrongObjectPtr to prevent GC.)
 	
-	const TStrongObjectPtr InstanceObject = TStrongObjectPtr(NewObject<UToolCallFakeClass>());
+	const TStrongObjectPtr InstanceObject = 
+		TStrongObjectPtr(NewObject<UToolCallFakeClass>());
 
 	InstanceObject->Float = ClassFloatValue;
 	
 	// Execute.
 	
-	const TOptional<FObjectFunctionToolCall> ToolCall = FObjectFunctionToolCall::FromNamedClassFunction(InstanceObject.Get(), FunctionName);
-	if (!TestTrue("Tool call created from class and function name.", ToolCall.IsSet()))
+	const TOptional<FObjectFunctionToolCall> ToolCall = 
+		FObjectFunctionToolCall::ConstructWithNamedClassFunction(InstanceObject.Get(), FunctionName);
+	if (!TestTrue("Tool call created from class and function name.", 
+			ToolCall.IsSet()))
 	{
 		return false;
 	}
 	
-	FObjectFunctionToolCall::FFunctionInputParamsJson JsonInputParamsJson;
-	JsonInputParamsJson.Set<FString>(FunctionInputParamsJsonString);
-	const FObjectFunctionToolCall::FToolCallResult ToolCallResult = ToolCall->Execute(JsonInputParamsJson);
+	const FObjectFunctionToolCall::FToolCallResult ToolCallResult = ToolCall->Execute();
 	
-	if (!TestTrue("Execute tool call result does not have error.", !ToolCallResult.HasError()))
+	if (!TestTrue("Execute tool call result does not have error.", 
+			!ToolCallResult.HasError()))
 	{
-		UE_LOG(LogAIAssistant, Error, TEXT("Tool call error - %s"), *ToolCallResult.GetError().ToString());
+		UE_LOG(LogAIAssistant, Error, TEXT("Tool call error - %s"), 
+			*ToolCallResult.GetError().ToString());
 		return false;
 	}
-	if (!TestTrue("Execute tool call result has value.", ToolCallResult.HasValue()))
+	if (!TestTrue("Execute tool call result has value.", 
+			ToolCallResult.HasValue()))
 	{
 		return false;
 	}
 
 	const TSharedPtr<FJsonObject> JsonObject = ToolCallResult.GetValue();
 
-	MaybePrintJsonToLog(TEXT("ToolCallReturnedJsonObject"), JsonObject.ToSharedRef());
-	
-	if (!TestTrue("Execute tool call result has valid value.", JsonObject.IsValid()))
+	if (!TestTrue("Execute tool call result has valid value.", 
+			JsonObject.IsValid()))
 	{
 		return false;
 	}
-	if (!TestTrue("Return object is not empty (has values.)", JsonObject->Values.Num() > 0))
+
+	MaybePrintJsonToLog(TEXT("ToolCallWithOptionalParamReturnedJsonObject"), 
+		JsonObject.ToSharedRef());
+	
+	if (!TestTrue("Return object is not empty (has values.)", 
+			JsonObject->Values.Num() > 0))
 	{
 		return false;
 	}	
-	if (!TestTrue(TEXT("Return object has float returnValue."), JsonObject->HasTypedField(TEXT("returnValue"), EJson::Number)))
+	if (!TestTrue(TEXT("Return object has float returnValue."), 
+			JsonObject->HasTypedField(TEXT("returnValue"), EJson::Number)))
 	{
 		return false;
 	}
 	
 	// Check.
 	
-	const float ReturnValueFloat = JsonObject->GetNumberField(TEXT("returnValue"));
-	const float ReturnValueFloatExpected = InstanceObject->TestFunc(FunctionInputFloatValue);
-	const FString TestMessage = FString::Printf(TEXT("Return value %f is expected value %f."), ReturnValueFloat, ReturnValueFloatExpected);
-	if (!TestTrue(TestMessage, ReturnValueFloat == ReturnValueFloatExpected))
+	const float ReturnValueFloat = 
+		JsonObject->GetNumberField(TEXT("returnValue"));
+	const float ReturnValueFloatExpected = 
+		InstanceObject->TestFuncWithOptionalParam();
+	const FString TestMessage = 
+		FString::Printf(TEXT("Return value %f is expected value %f."), ReturnValueFloat, ReturnValueFloatExpected);
+	if (!TestTrue(TestMessage, 
+			ReturnValueFloat == ReturnValueFloatExpected))
 	{
 		return false;
 	}
@@ -109,10 +206,8 @@ bool FAIAssistantToolCallExecuteToolCallFunction::RunTest(const FString& Paramet
 	return true;
 }
 
-
 #endif
 
-
 PRAGMA_ENABLE_DEPRECATION_WARNINGS
 
 
diff --git a/Engine/Plugins/Experimental/AIAssistant/Source/AIAssistant/Private/Tests/AIAssistantToolCallTest.h b/Engine/Plugins/Experimental/AIAssistant/Source/AIAssistant/Private/Tests/AIAssistantToolCallTest.h
--- a/Engine/Plugins/Experimental/AIAssistant/Source/AIAssistant/Private/Tests/AIAssistantToolCallTest.h
+++ b/Engine/Plugins/Experimental/AIAssistant/Source/AIAssistant/Private/Tests/AIAssistantToolCallTest.h
@@ -17,7 +17,13 @@ public:
 	float Float = 0.0f;
 	
 	UFUNCTION(BlueprintCallable, Category="Main") 
-	float TestFunc(const float InFloat) const 
+	float TestFuncWithRequiredParam(const float InFloat) const 
+	{
+		return 100.0f + Float + InFloat; 
+	}
+	
+	UFUNCTION(BlueprintCallable, Category="Main") 
+	float TestFuncWithOptionalParam(const float InFloat = 3.0f) const 
 	{
 		return 100.0f + Float + InFloat; 
 	}
