[AI Assistant] Toolsets: Add more testing and validation
- remove some caching from trampoline
- add BP callable method to list tool names
- toolsets with no valid tools will not be registered (currently just checks for having at least one tool, will add tool schema validation as methods become available)
- add tests for toolsets with default params, invalid params, no tools, TMap, TArray, TSet params

#rnx
#Synced-CL 49446976

diff --git a/Engine/Plugins/Experimental/AIAssistant/Source/AIAssistant/Private/AIAssistantToolset.cpp b/Engine/Plugins/Experimental/AIAssistant/Source/AIAssistant/Private/AIAssistantToolset.cpp
--- a/Engine/Plugins/Experimental/AIAssistant/Source/AIAssistant/Private/AIAssistantToolset.cpp
+++ b/Engine/Plugins/Experimental/AIAssistant/Source/AIAssistant/Private/AIAssistantToolset.cpp
@@ -34,12 +34,7 @@ namespace UE::AIAssistant
 			: ToolsetHandler(InToolsetHandler)
 		{
 			TStrongObjectPtr<UAIAssistantToolset> Handler = ToolsetHandler.Pin();
-			if (Handler)
-			{
-				ToolsetName = Handler->GetName();
-				ToolsetVersion = Handler->GetVersion();
-				ToolsetSchema = Handler->GetJsonSchema();
-			}
+			ToolsetName = Handler ? Handler->GetName() : FString();
 		}
 
 		// Executes a tool with the given name and JSON input, returning a future with the JSON
@@ -74,7 +69,8 @@ namespace UE::AIAssistant
 		// Returns the JSON schema for the toolset
 		virtual FString GetToolsetJsonSchema() const override
 		{
-			return ToolsetSchema;
+			TStrongObjectPtr<UAIAssistantToolset> Handler = ToolsetHandler.Pin();
+			return Handler ? Handler->GetJsonSchema() : FString();
 		}
 
 		// Returns the name of the toolset
@@ -86,7 +82,8 @@ namespace UE::AIAssistant
 		// Returns the version of the toolset
 		virtual FString GetToolsetVersion() const override
 		{
-			return ToolsetVersion;
+			TStrongObjectPtr<UAIAssistantToolset> Handler = ToolsetHandler.Pin();
+			return Handler ? Handler->GetVersion(): FString();
 		}
 
 	  private:
@@ -97,18 +94,13 @@ namespace UE::AIAssistant
 		// into its map of handlers.  Cannot ask the owning UAIAssistantToolset for its name at that point, as
 		// it's being destroyed.
 		FString ToolsetName;
-
-		// Cached version, since caching the name was already necessary, might as well cache this as well.
-		FString ToolsetVersion;
-
-		// Cached schema, since caching the name was already necessary, might as well cache this as well.
-		FString ToolsetSchema;
 	};
 }
 
 void UAIAssistantToolset::PostInitProperties()
 {
 	Super::PostInitProperties();
+
 	// Only actual instances should register.
 	if (!HasAnyFlags(RF_ClassDefaultObject | RF_ArchetypeObject))
 	{
@@ -128,43 +120,58 @@ void UAIAssistantToolset::BeginDestroy()
 
 FString UAIAssistantToolset::GetJsonSchema() const
 {
-	FString ToolsetName = GetName();
 	FString OutputString;
-	TSharedPtr<FJsonObject> Schema;
+	TSharedPtr<FJsonObject> Schema = GenerateSchema();
 
-	Schema = MakeShared<FJsonObject>();
-	check(Schema);
+	if (!FJsonSerializer::Serialize(
+			Schema.ToSharedRef(),
+			TJsonWriterFactory<TCHAR, TCondensedJsonPrintPolicy<TCHAR>>::Create(&OutputString)))
+	{
+		UE_LOG(
+			LogAIAssistant, Warning,
+			TEXT("Failed to generate valid toolset schema for toolset %s"), *GetName());
+	}
+
+	return OutputString;
+}
 
-	Schema->SetStringField(FString(TEXT("name")), ToolsetName);
+TSharedPtr<FJsonObject> UAIAssistantToolset::GenerateSchema() const
+{
+	TSharedPtr<FJsonObject> Schema = MakeShared<FJsonObject>();
+	check(Schema);
+	Schema->SetStringField(FString(TEXT("name")), GetName());
 	Schema->SetStringField(FString(TEXT("version")), GetVersion());
 	Schema->SetStringField(FString(TEXT("description")), GetDescription());
 
 	TArray<TSharedPtr<FJsonValue>> ToolsArray;
-	TArray<TObjectPtr<UFunction>> ValidTools = GetValidToolMethods();
+	TArray<TObjectPtr<UFunction>> ValidTools = GetValidToolMethodsList();
 	for (UFunction* Function : ValidTools)
 	{
-        if (Function)
-		{
+        check(Function);
 
 PRAGMA_DISABLE_DEPRECATION_WARNINGS
-			TSharedPtr<FJsonObject> ToolEntry = FJsonSchemaGenerator::UStructToJsonSchemaObject(Function);
+		TSharedPtr<FJsonObject> ToolEntry = FJsonSchemaGenerator::UStructToJsonSchemaObject(Function);
 PRAGMA_ENABLE_DEPRECATION_WARNINGS
 
-			ToolsArray.Add(MakeShared<FJsonValueObject>(ToolEntry));
-		}
+		ToolsArray.Add(MakeShared<FJsonValueObject>(ToolEntry));
 	}
 	Schema->SetArrayField(FString(TEXT("tools")), ToolsArray);
+	return Schema;
+}
 
-	if (!FJsonSerializer::Serialize(
-			Schema.ToSharedRef(),
-			TJsonWriterFactory<TCHAR, TCondensedJsonPrintPolicy<TCHAR>>::Create(&OutputString)))
+void UAIAssistantToolset::GenerateToolCallObjects()
+{
+	Tools.Empty();
+	TArray<TObjectPtr<UFunction>> ValidTools = GetValidToolMethodsList();
+	for (UFunction* Function : ValidTools)
 	{
-		UE_LOG(
-			LogAIAssistant, Warning,
-			TEXT("Failed to generate valid toolset schema for toolset %s"), *ToolsetName);
+		check(Function);
+		const UE::AIAssistant::FObjectFunctionToolCall ToolCall =
+			UE::AIAssistant::FObjectFunctionToolCall(
+				TNotNull<UObject*>(this), TNotNull<UFunction*>(Function));
+		Tools.Add(
+			Function->GetName(), MakeShared<UE::AIAssistant::FObjectFunctionToolCall>(ToolCall));
 	}
-
-	return OutputString;
 }
 
 void UAIAssistantToolset::Register()
@@ -180,6 +187,13 @@ void UAIAssistantToolset::Register()
 			return;
 		}
 
+		if (!HasValidTools())
+		{
+			UE_LOG(LogAIAssistant, Warning,
+				   TEXT("Unable to register Toolset %s: invalid tool list"), *ToolsetName);
+			return;
+		}
+
 		if (!RegistryHandler.IsValid())
 		{
 			RegistryHandler = MakeShared<UE::AIAssistant::FToolsetHandlerTrampoline>(
@@ -221,25 +235,26 @@ void UAIAssistantToolset::InvokeTool(
 	OutResult.Empty();
 	OutError.Empty();
 
-	if (RegisteredToolNames.IsEmpty())
+	// Ensure valid tool list is up to date
+	if (Tools.IsEmpty())
 	{
-		RegisterToolNames();
+		GenerateToolCallObjects();
 	}
 
-	UFunction* Function = GetClass()->FindFunctionByName(*ToolName);
-	if (!Function || !RegisteredToolNames.Contains(ToolName))
+	TSharedPtr<UE::AIAssistant::FObjectFunctionToolCall>* ToolCallPtr = Tools.Find(ToolName);
+	if (!ToolCallPtr)
 	{
 		OutError = FString::Printf(TEXT("Unknown tool %s"), *ToolName);
 		return;
 	}
+	TSharedPtr<UE::AIAssistant::FObjectFunctionToolCall> ToolCall = *ToolCallPtr;
+	check(ToolCall);
 
 	TSharedPtr<FJsonObject> ToolCallReturnValue;
 	UE::AIAssistant::FObjectFunctionToolCall::FFunctionInputParamsJson JsonParams;
 	JsonParams.Set<FString>(JsonInput);
-	const UE::AIAssistant::FObjectFunctionToolCall ToolCall =
-		UE::AIAssistant::FObjectFunctionToolCall(TNotNull<UObject*>(this), TNotNull<UFunction*>(Function));
 
-	UE::AIAssistant::FObjectFunctionToolCall::FToolCallResult ToolCallResult = ToolCall.Execute(JsonParams);
+	UE::AIAssistant::FObjectFunctionToolCall::FToolCallResult ToolCallResult = ToolCall->Execute(JsonParams);
 
 	if (ToolCallResult.HasError())
 	{
@@ -272,20 +287,29 @@ void UAIAssistantToolset::InvokeTool(
 	return;
 }
 
-void UAIAssistantToolset::RegisterToolNames()
+TArray<FString> UAIAssistantToolset::GetToolNames()
 {
-	TArray<TObjectPtr<UFunction>> ValidTools = GetValidToolMethods();
+	TArray<TObjectPtr<UFunction>> ValidTools = GetValidToolMethodsList();
+	TArray<FString> ToolNames;
+	for (TObjectPtr<UFunction> FuncPtr : ValidTools)
+	{
+		ToolNames.Add(FuncPtr->GetName());
+	}
+	return ToolNames;
+}
 
-	for (UFunction* Function : ValidTools)
+bool UAIAssistantToolset::HasValidTools()
+{
+	GenerateToolCallObjects();
+	if (Tools.IsEmpty())
 	{
-        if (Function)
-		{
-			RegisteredToolNames.Add(Function->GetName());
-		}
+		return false;
 	}
+	// TODO: when we have validation methods for tool schemas we can check them.
+	return true;
 }
 
-TArray<TObjectPtr<UFunction>> UAIAssistantToolset::GetValidToolMethods() const
+TArray<TObjectPtr<UFunction>> UAIAssistantToolset::GetValidToolMethodsList() const
 {
 	TArray<TObjectPtr<UFunction>> ValidTools;
 	for (TFieldIterator<UFunction> FuncIt(GetClass(), EFieldIterationFlags::None); FuncIt; ++FuncIt)
diff --git a/Engine/Plugins/Experimental/AIAssistant/Source/AIAssistant/Private/AIAssistantToolset.h b/Engine/Plugins/Experimental/AIAssistant/Source/AIAssistant/Private/AIAssistantToolset.h
--- a/Engine/Plugins/Experimental/AIAssistant/Source/AIAssistant/Private/AIAssistantToolset.h
+++ b/Engine/Plugins/Experimental/AIAssistant/Source/AIAssistant/Private/AIAssistantToolset.h
@@ -5,7 +5,9 @@
 #include "CoreMinimal.h"
 
 #include "Containers/UnrealString.h"
+#include "Dom/JsonObject.h"
 
+#include "AIAssistantToolCall.h"
 #include "AIAssistantToolsetHandler.h"
 
 #include "AIAssistantToolset.generated.h"
@@ -77,6 +79,10 @@ public:
 	UFUNCTION(BlueprintCallable, Category=AssistantToolset)
 	void InvokeTool(const FString& ToolName, const FString& JsonInput, FString& OutResult, FString& OutError);
 
+	// Return list of tool names in this toolset.
+	UFUNCTION(BlueprintCallable, Category=AssistantToolset)
+	TArray<FString> GetToolNames();
+
 	// Overridden to automatically register toolset with AI Assistant subsystem.
 	// Subclasses should also override the ToolsetVersion property here.
 	virtual void PostInitProperties() override;
@@ -84,16 +90,23 @@ public:
 	virtual void BeginDestroy() override;
 
 private:
-	// Generate list of valid tool names.
-	void RegisterToolNames();
 
 	// Handler that implements IToolsetHandler to register with the AI Assistant.
 	TSharedPtr<UE::AIAssistant::FToolsetHandlerTrampoline> RegistryHandler;
 
-	// Generate list of valid tools.
-	TArray<TObjectPtr<UFunction>> GetValidToolMethods() const;
+	// Map of valid tool names to tool call objects.
+	TMap<FString, TSharedPtr<UE::AIAssistant::FObjectFunctionToolCall>> Tools;
+
+	// Whether the toolset has a valid set of tools.
+	bool HasValidTools();
+
+	// Generate schema from reflection data.
+	TSharedPtr<FJsonObject> GenerateSchema() const;
+
+	// Generate Tools list.
+	void GenerateToolCallObjects();
 
-	// Array of valid tool names present in toolset schema.
-	TArray<FString> RegisteredToolNames;
+	// Helper method to generate list of valid tools.
+	TArray<TObjectPtr<UFunction>> GetValidToolMethodsList() const;
 };
 
diff --git a/Engine/Plugins/Experimental/AIAssistant/Source/AIAssistant/Private/Tests/AIAssistantTestSchemaBuilder.cpp b/Engine/Plugins/Experimental/AIAssistant/Source/AIAssistant/Private/Tests/AIAssistantTestSchemaBuilder.cpp
new file mode 100644
--- /dev/null
+++ b/Engine/Plugins/Experimental/AIAssistant/Source/AIAssistant/Private/Tests/AIAssistantTestSchemaBuilder.cpp
@@ -0,0 +1,110 @@
+// Copyright Epic Games, Inc. All Rights Reserved.
+
+#include "AIAssistantTestSchemaBuilder.h"
+
+#include "CoreMinimal.h"
+#include "JsonDomBuilder.h"
+
+FJsonDomBuilder::FObject MakeSchemaBuilder(
+	const FString& Name, const FString& Version, const FString& Description)
+{
+	FJsonDomBuilder::FObject SchemaBuilder;
+	SchemaBuilder.Set(TEXT("name"), Name);
+	SchemaBuilder.Set(TEXT("version"), Version);
+	SchemaBuilder.Set(TEXT("description"), Description);
+	return SchemaBuilder;
+}
+
+FJsonDomBuilder::FArray AddToolSchema(
+	FJsonDomBuilder::FObject& Builder, FJsonDomBuilder::FArray& ToolsArray,
+	FJsonDomBuilder::FObject& ToolEntry)
+{
+	ToolsArray.Add(ToolEntry);
+	Builder.Set(TEXT("tools"), ToolsArray);
+	return ToolsArray;
+}
+
+FJsonDomBuilder::FObject MakeToolSchemaEntry(
+	const FString& Name, const FString& Description,
+	FJsonDomBuilder::FObject& InputSchema, FJsonDomBuilder::FObject& OutputSchema)
+{
+	FJsonDomBuilder::FObject ToolSchema;
+	ToolSchema.Set(TEXT("description"), Description);
+	ToolSchema.Set(TEXT("name"), Name);
+	ToolSchema.Set(TEXT("inputSchema"), InputSchema);
+	ToolSchema.Set(TEXT("outputSchema"), OutputSchema);
+	return ToolSchema;
+}
+
+FJsonDomBuilder::FObject MakeSchema()
+{
+	FJsonDomBuilder::FObject Schema;
+	Schema.Set(TEXT("type"), TEXT("object"));
+	return Schema;
+}
+
+FJsonDomBuilder::FObject MakeProperty(const FString& Type, TSharedPtr<FJsonValue> DefaultValue)
+{
+	FJsonDomBuilder::FObject PropertyType;
+	PropertyType.Set(TEXT("type"), Type);
+	if (DefaultValue)
+	{
+		PropertyType.Set(TEXT("default"), DefaultValue);
+	}
+	return PropertyType;
+}
+
+FJsonDomBuilder::FObject MakeArrayProperty(const FString& Type)
+{
+	FJsonDomBuilder::FObject PropertyType;
+	PropertyType.Set(TEXT("type"), TEXT("array"));
+	FJsonDomBuilder::FObject ItemType;
+	ItemType.Set(TEXT("type"), Type);
+	PropertyType.Set(TEXT("items"), ItemType);
+	return PropertyType;
+}
+
+FJsonDomBuilder::FObject MakeSetProperty(const FString& Type)
+{
+	FJsonDomBuilder::FObject PropertyType;
+	PropertyType.Set(TEXT("type"), TEXT("array"));
+	PropertyType.Set(TEXT("uniqueItems"), true);
+	FJsonDomBuilder::FObject ItemType;
+	ItemType.Set(TEXT("type"), Type);
+	PropertyType.Set(TEXT("items"), ItemType);
+	return PropertyType;
+}
+
+FJsonDomBuilder::FObject MakeMapProperty(const FString& Type)
+{
+	FJsonDomBuilder::FObject PropertyType;
+	PropertyType.Set(TEXT("type"), TEXT("object"));
+	FJsonDomBuilder::FObject ItemType;
+	ItemType.Set(TEXT("type"), Type);
+	PropertyType.Set(TEXT("additionalProperties"), ItemType);
+	return PropertyType;
+}
+
+FJsonDomBuilder::FObject MakeNamedProperty(
+	const FString& Name, const FString& Type, TSharedPtr<FJsonValue> DefaultValue)
+{
+	FJsonDomBuilder::FObject Property;
+	FJsonDomBuilder::FObject PropertyType = MakeProperty(Type, DefaultValue);
+	Property.Set(Name, PropertyType);
+	return Property;
+}
+
+void AddPropertiesToSchema(FJsonDomBuilder::FObject& Schema, FJsonDomBuilder::FObject Properties)
+{
+	Schema.Set(TEXT("properties"), Properties);
+}
+
+void AddRequiredListToSchema(FJsonDomBuilder::FObject& Schema, const TArray<FString>& RequiredList)
+{
+	FJsonDomBuilder::FArray RequiredArray;
+	for (FString Item : RequiredList)
+	{
+		RequiredArray.Add(Item);
+	}
+	Schema.Set(TEXT("required"), RequiredArray);
+}
diff --git a/Engine/Plugins/Experimental/AIAssistant/Source/AIAssistant/Private/Tests/AIAssistantTestSchemaBuilder.h b/Engine/Plugins/Experimental/AIAssistant/Source/AIAssistant/Private/Tests/AIAssistantTestSchemaBuilder.h
new file mode 100644
--- /dev/null
+++ b/Engine/Plugins/Experimental/AIAssistant/Source/AIAssistant/Private/Tests/AIAssistantTestSchemaBuilder.h
@@ -0,0 +1,53 @@
+// Copyright Epic Games, Inc. All Rights Reserved.
+
+#pragma once
+
+#include "CoreMinimal.h"
+#include "JsonDomBuilder.h"
+
+/// Utility functions for generating toolset JSON schemas.
+
+// Call this first to create a base toolset schema object with the toolset's name, version, and description.
+// Call AsJsonObject on it when you are done to generate a TSharedRef<FJsonObject> from it.
+FJsonDomBuilder::FObject MakeSchemaBuilder(
+	const FString& Name, const FString& Version, const FString& Description);
+
+// Add a tool schema to the tools array and add the tools array to the schema builder.  Returns the tools
+// array if you want to add more tool entries to it.
+FJsonDomBuilder::FArray AddToolSchema(
+	FJsonDomBuilder::FObject& Builder, FJsonDomBuilder::FArray& ToolsArray,
+	FJsonDomBuilder::FObject& ToolEntry);
+
+// Make a tool schema with the tool's name and description, input schema, and output schema.
+FJsonDomBuilder::FObject MakeToolSchemaEntry(
+	const FString& Name, const FString& Description,
+	FJsonDomBuilder::FObject& InputSchema, FJsonDomBuilder::FObject& OutputSchema);
+
+// Make a schema object that you can use for either an input or output schema.
+FJsonDomBuilder::FObject MakeSchema();
+
+// Make a property object for putting in an input or output schema, specifying its type and an optional
+// default value.
+FJsonDomBuilder::FObject MakeProperty(
+	const FString& Type, TSharedPtr<FJsonValue> DefaultValue = nullptr);
+
+// Make an array property object for putting in an input or output schema, specifying the type of its items.
+FJsonDomBuilder::FObject MakeArrayProperty(const FString& Type);
+
+// Make a set property object for putting in an input or output schema, specifying the type of its items.
+FJsonDomBuilder::FObject MakeSetProperty(const FString& Type);
+
+// Make a map property object for putting in an input or output schema, specifying the type of its items.
+FJsonDomBuilder::FObject MakeMapProperty(const FString& Type);
+
+// Make a property object for putting in an input or output schema, specifying its name, type and an optional
+// default value.
+FJsonDomBuilder::FObject MakeNamedProperty(
+	const FString& Name, const FString& Type, TSharedPtr<FJsonValue> DefaultValue = nullptr);
+
+// Add a property object to an input or output schema.
+void AddPropertiesToSchema(FJsonDomBuilder::FObject& Schema, FJsonDomBuilder::FObject Properties);
+
+// Add a required list of names to an input or output schema.
+void AddRequiredListToSchema(FJsonDomBuilder::FObject& Schema, const TArray<FString>& RequiredList);
+
diff --git a/Engine/Plugins/Experimental/AIAssistant/Source/AIAssistant/Private/Tests/AIAssistantToolsetTest.cpp b/Engine/Plugins/Experimental/AIAssistant/Source/AIAssistant/Private/Tests/AIAssistantToolsetTest.cpp
--- a/Engine/Plugins/Experimental/AIAssistant/Source/AIAssistant/Private/Tests/AIAssistantToolsetTest.cpp
+++ b/Engine/Plugins/Experimental/AIAssistant/Source/AIAssistant/Private/Tests/AIAssistantToolsetTest.cpp
@@ -13,17 +13,19 @@
 #include "Serialization/JsonWriter.h"
 #include "Serialization/JsonSerializer.h"
 
-#include "Tests/AIAssistantTestFlags.h"
 #include "AIAssistantBPToolsetRegistry.h"
 #include "AIAssistantLog.h"
+#include "Tests/AIAssistantTestSchemaBuilder.h"
+#include "Tests/AIAssistantTestFlags.h"
 
 
 #if WITH_DEV_AUTOMATION_TESTS
 
 const FString FakeToolsetName(TEXT("FakeToolset"));
-const FString FakeToolset2Name(TEXT("FakeToolset2"));
+const FString FakeToolsetUpdatesVersionName(TEXT("FakeToolsetUpdatesVersion"));
 const FString FakeToolsetInput(TEXT(R"({"inString":"hello"})"));
 const FString FakeToolName(TEXT("SimpleTool"));
+const FString ReturnValueName(TEXT("returnValue"));
 
 // Wrapper for creating and cleaning up UObjects.
 template<typename T>class ToolsetContainer
@@ -31,6 +33,7 @@ template<typename T>class ToolsetContainer
 public:
 	ToolsetContainer()
 	{
+		CollectGarbage(GARBAGE_COLLECTION_KEEPFLAGS);
 		Instance = NewObject<T>();
 	}
 
@@ -46,10 +49,176 @@ public:
 		return Instance;
 	}
 
+	FString Input;
+	FString Result;
+	FString Error;
+
 private:
 	TObjectPtr<T> Instance;
 };
 
+TSharedPtr<FJsonObject> UFakeToolset::GetExpectedSchema()
+{
+	FJsonDomBuilder::FObject SchemaBuilder = MakeSchemaBuilder(GetName(), GetVersion(), GetDescription());
+
+	FJsonDomBuilder::FObject InputSchema = MakeSchema();
+	FJsonDomBuilder::FObject Properties;
+
+	FString ParamName = TEXT("inString");
+
+	Properties.Set(ParamName, MakeProperty(TEXT("string")));
+	AddPropertiesToSchema(InputSchema, Properties);
+	AddRequiredListToSchema(InputSchema, TArray<FString>({ParamName}));
+
+	FJsonDomBuilder::FObject OutputSchema = MakeSchema();
+	AddPropertiesToSchema(OutputSchema, MakeNamedProperty(ReturnValueName, TEXT("string")));
+	AddRequiredListToSchema(OutputSchema, TArray<FString>({ReturnValueName}));
+
+	FJsonDomBuilder::FObject ToolEntry = MakeToolSchemaEntry(
+		TEXT("FakeToolset.SimpleTool"), TEXT("Simple tool that does nothing."),
+		InputSchema, OutputSchema);
+
+	FJsonDomBuilder::FArray ToolsArray;
+	AddToolSchema(SchemaBuilder, ToolsArray, ToolEntry);
+
+	return SchemaBuilder.AsJsonObject();
+}
+
+TSharedPtr<FJsonObject> UFakeToolsetWithDefaultParams::GetExpectedSchema()
+{
+	FJsonDomBuilder::FObject SchemaBuilder = MakeSchemaBuilder(GetName(), GetVersion(), GetDescription());
+
+	FJsonDomBuilder::FObject InputSchema = MakeSchema();
+	FJsonDomBuilder::FObject Properties;
+
+	FString ParamName1 = TEXT("inString");
+	FString ParamName2 = TEXT("inValue");
+
+	Properties.Set(ParamName1, MakeProperty(TEXT("string")));
+	Properties.Set(ParamName2, MakeProperty(TEXT("integer"), MakeShared<FJsonValueNumber>(42)));
+	AddPropertiesToSchema(InputSchema, Properties);
+	AddRequiredListToSchema(InputSchema, TArray<FString>({ParamName1}));
+
+	FJsonDomBuilder::FObject OutputSchema = MakeSchema();
+	AddPropertiesToSchema(OutputSchema, MakeNamedProperty(ReturnValueName, TEXT("string")));
+	AddRequiredListToSchema(OutputSchema, TArray<FString>({ReturnValueName}));
+
+	FJsonDomBuilder::FObject ToolEntry = MakeToolSchemaEntry(
+		TEXT("FakeToolsetWithDefaultParams.ToolWithDefaults"),
+		TEXT("Method with some params with default values."),
+		InputSchema, OutputSchema);
+
+	FJsonDomBuilder::FArray ToolsArray;
+	AddToolSchema(SchemaBuilder, ToolsArray, ToolEntry);
+
+	return SchemaBuilder.AsJsonObject();
+}
+
+TSharedPtr<FJsonObject> UFakeToolsetWithNoTools::GetExpectedSchema()
+{
+	FJsonDomBuilder::FObject SchemaBuilder = MakeSchemaBuilder(GetName(), GetVersion(), GetDescription());
+
+	FJsonDomBuilder::FArray ToolsArray;
+	SchemaBuilder.Set(TEXT("tools"), ToolsArray);
+	return SchemaBuilder.AsJsonObject();
+}
+
+TSharedPtr<FJsonObject> UFakeToolsetWithArrayParams::GetExpectedSchema()
+{
+	FJsonDomBuilder::FObject SchemaBuilder = MakeSchemaBuilder(GetName(), GetVersion(), GetDescription());
+
+	FJsonDomBuilder::FObject InputSchema = MakeSchema();
+
+	FString ParamName1 = TEXT("arrayParamOne");
+	FString ParamName2 = TEXT("arrayParamTwo");
+
+	FJsonDomBuilder::FObject InputProperties;
+	InputProperties.Set(ParamName1, MakeArrayProperty(TEXT("integer")));
+	InputProperties.Set(ParamName2, MakeArrayProperty(TEXT("string")));
+	AddPropertiesToSchema(InputSchema, InputProperties);
+	AddRequiredListToSchema(InputSchema, TArray<FString>({ParamName1, ParamName2}));
+
+	FJsonDomBuilder::FObject OutputSchema = MakeSchema();
+	FJsonDomBuilder::FObject OutputProperties;
+	OutputProperties.Set(ReturnValueName, MakeArrayProperty(TEXT("number")));
+	AddPropertiesToSchema(OutputSchema, OutputProperties);
+	AddRequiredListToSchema(OutputSchema, TArray<FString>({ReturnValueName}));
+
+	FJsonDomBuilder::FObject ToolEntry = MakeToolSchemaEntry(
+		TEXT("FakeToolsetWithArrayParams.ToolWithArrays"),
+		TEXT("Method with some params with TArray values."),
+		InputSchema, OutputSchema);
+
+	FJsonDomBuilder::FArray ToolsArray;
+	AddToolSchema(SchemaBuilder, ToolsArray, ToolEntry);
+
+	return SchemaBuilder.AsJsonObject();
+}
+
+TSharedPtr<FJsonObject> UFakeToolsetWithSetParams::GetExpectedSchema()
+{
+	FJsonDomBuilder::FObject SchemaBuilder = MakeSchemaBuilder(GetName(), GetVersion(), GetDescription());
+
+	FJsonDomBuilder::FObject InputSchema = MakeSchema();
+
+	FString ParamName1 = TEXT("setParamOne");
+	FString ParamName2 = TEXT("setParamTwo");
+
+	FJsonDomBuilder::FObject InputProperties;
+	InputProperties.Set(ParamName1, MakeSetProperty(TEXT("integer")));
+	InputProperties.Set(ParamName2, MakeSetProperty(TEXT("string")));
+	AddPropertiesToSchema(InputSchema, InputProperties);
+	AddRequiredListToSchema(InputSchema, TArray<FString>({ParamName1, ParamName2}));
+
+	FJsonDomBuilder::FObject OutputSchema = MakeSchema();
+	FJsonDomBuilder::FObject OutputProperties;
+	OutputProperties.Set(ReturnValueName, MakeSetProperty(TEXT("number")));
+	AddPropertiesToSchema(OutputSchema, OutputProperties);
+	AddRequiredListToSchema(OutputSchema, TArray<FString>({ReturnValueName}));
+
+	FJsonDomBuilder::FObject ToolEntry = MakeToolSchemaEntry(
+		TEXT("FakeToolsetWithSetParams.ToolWithSets"),
+		TEXT("Method with some params with TSet values."),
+		InputSchema, OutputSchema);
+
+	FJsonDomBuilder::FArray ToolsArray;
+	AddToolSchema(SchemaBuilder, ToolsArray, ToolEntry);
+
+	return SchemaBuilder.AsJsonObject();
+}
+
+TSharedPtr<FJsonObject> UFakeToolsetWithMapParams::GetExpectedSchema()
+{
+	FJsonDomBuilder::FObject SchemaBuilder = MakeSchemaBuilder(GetName(), GetVersion(), GetDescription());
+
+	FJsonDomBuilder::FObject InputSchema = MakeSchema();
+
+	FString ParamName1 = TEXT("mapParamOne");
+	FString ParamName2 = TEXT("mapParamTwo");
+
+	FJsonDomBuilder::FObject InputProperties;
+	InputProperties.Set(ParamName1, MakeMapProperty(TEXT("integer")));
+	InputProperties.Set(ParamName2, MakeMapProperty(TEXT("string")));
+	AddPropertiesToSchema(InputSchema, InputProperties);
+	AddRequiredListToSchema(InputSchema, TArray<FString>({ParamName1, ParamName2}));
+
+	FJsonDomBuilder::FObject OutputSchema = MakeSchema();
+	FJsonDomBuilder::FObject OutputProperties;
+	OutputProperties.Set(ReturnValueName, MakeMapProperty(TEXT("number")));
+	AddPropertiesToSchema(OutputSchema, OutputProperties);
+	AddRequiredListToSchema(OutputSchema, TArray<FString>({ReturnValueName}));
+
+	FJsonDomBuilder::FObject ToolEntry = MakeToolSchemaEntry(
+		TEXT("FakeToolsetWithMapParams.ToolWithMaps"),
+		TEXT("Method with some params with TMap values."),
+		InputSchema, OutputSchema);
+
+	FJsonDomBuilder::FArray ToolsArray;
+	AddToolSchema(SchemaBuilder, ToolsArray, ToolEntry);
+
+	return SchemaBuilder.AsJsonObject();
+}
+
 // Test that toolset class creation registers toolset and destroy unregisters toolset.
 IMPLEMENT_SIMPLE_AUTOMATION_TEST(
 	FAIAssistantToolsetRegistersOnCreateAndDestroy,
@@ -75,6 +244,24 @@ bool FAIAssistantToolsetRegistersOnCreateAndDestroy::RunTest(const FString& Unus
 	return true;
 }
 
+// Test that toolset with no tools fails to register.
+IMPLEMENT_SIMPLE_AUTOMATION_TEST(
+	FAIAssistantToolsetNoToolsRegistrationRejected,
+	"AI.Assistant.Toolset.NoToolsRegistrationRejected",
+	AIAssistantTest::Flags)
+
+bool FAIAssistantToolsetNoToolsRegistrationRejected::RunTest(const FString& UnusedParameters)
+{
+	ToolsetContainer<UFakeToolsetWithNoTools> Container;
+	TObjectPtr<UFakeToolsetWithNoTools> Instance = *Container;
+	TestFalse("Toolset with no tools is not registered.",
+			  UAIAssistantToolsetRegistry::IsToolsetRegistered(Instance->GetName()));
+	Instance->Register();
+	TestFalse("Toolset with no tools is not registered when asking explicitly.",
+			  UAIAssistantToolsetRegistry::IsToolsetRegistered(Instance->GetName()));
+	return true;
+}
+
 // Test that toolset can return its name, version, and description.
 IMPLEMENT_SIMPLE_AUTOMATION_TEST(
 	FAIAssistantToolsetCanGetValues,
@@ -83,8 +270,6 @@ IMPLEMENT_SIMPLE_AUTOMATION_TEST(
 
 bool FAIAssistantToolsetCanGetValues::RunTest(const FString& UnusedParameters)
 {
-	CollectGarbage(GARBAGE_COLLECTION_KEEPFLAGS);
-
 	ToolsetContainer<UFakeToolset> Container;
 	TObjectPtr<UFakeToolset> Instance = *Container;
 	TestEqual("Toolset returns its own name",
@@ -104,65 +289,18 @@ IMPLEMENT_SIMPLE_AUTOMATION_TEST(
 
 bool FAIAssistantToolsetCanOverrideProperties::RunTest(const FString& UnusedParameters)
 {
-	CollectGarbage(GARBAGE_COLLECTION_KEEPFLAGS);
-
-	ToolsetContainer<UFakeToolset2> Container;
-	TObjectPtr<UFakeToolset2> Instance = *Container;
+	ToolsetContainer<UFakeToolsetUpdatesVersion> Container;
+	TObjectPtr<UFakeToolsetUpdatesVersion> Instance = *Container;
 	TestEqual("Toolset returns its own name",
-			  Instance->GetName(), FakeToolset2Name);
+			  Instance->GetName(), FakeToolsetUpdatesVersionName);
 	TestEqual("Toolset returns overridden version",
-			  Instance->GetVersion(), UFakeToolset2::MY_VERSION);
+			  Instance->GetVersion(), UFakeToolsetUpdatesVersion::MY_VERSION);
 	TestEqual("Toolset returns description from class definition",
 			  Instance->GetDescription(),
 			  FString(TEXT("Fake Toolset that does nothing except update version.")));
 	return true;
 }
 
-TSharedPtr<FJsonObject> GetFakeToolsetSchema(const TObjectPtr<UFakeToolset> Toolset)
-{
-	FJsonDomBuilder::FObject SchemaBuilder;
-
-	SchemaBuilder.Set(TEXT("name"), Toolset->GetName());
-	SchemaBuilder.Set(TEXT("version"), Toolset->GetVersion());
-	SchemaBuilder.Set(TEXT("description"), Toolset->GetDescription());
-
-	FJsonDomBuilder::FArray ToolsArray;
-	{
-		FJsonDomBuilder::FObject ToolEntry;
-		ToolEntry.Set(TEXT("description"), TEXT("Simple tool that does nothing."));
-		ToolEntry.Set(TEXT("name"), TEXT("FakeToolset.SimpleTool"));
-		{
-			FJsonDomBuilder::FObject InputSchema;
-			InputSchema.Set(TEXT("type"), TEXT("object"));
-			FJsonDomBuilder::FObject SchemaProps;
-			FJsonDomBuilder::FObject SchemaPropsType;
-			SchemaPropsType.Set(TEXT("type"), TEXT("string"));
-			SchemaProps.Set(TEXT("inString"), SchemaPropsType);
-			InputSchema.Set(TEXT("properties"), SchemaProps);
-			FJsonDomBuilder::FArray RequiredArray;
-			RequiredArray.Add(TEXT("inString"));
-			InputSchema.Set(TEXT("required"), RequiredArray);
-			ToolEntry.Set(TEXT("inputSchema"), InputSchema);
-		}
-		{
-			FJsonDomBuilder::FObject OutputSchema;
-			OutputSchema.Set(TEXT("type"), TEXT("object"));
-			FJsonDomBuilder::FObject SchemaProps;
-			FJsonDomBuilder::FObject SchemaPropsType;
-			SchemaPropsType.Set(TEXT("type"), TEXT("string"));
-			SchemaProps.Set(TEXT("returnValue"), SchemaPropsType);
-			OutputSchema.Set(TEXT("properties"), SchemaProps);
-			FJsonDomBuilder::FArray RequiredArray;
-			RequiredArray.Add(TEXT("returnValue"));
-			OutputSchema.Set(TEXT("required"), RequiredArray);
-			ToolEntry.Set(TEXT("outputSchema"), OutputSchema);
-		}
-		ToolsArray.Add(ToolEntry);
-	}
-	SchemaBuilder.Set(TEXT("tools"), ToolsArray);
-	return SchemaBuilder.AsJsonObject();
-}
-
 // Test that toolset can return a schema properly.
 IMPLEMENT_SIMPLE_AUTOMATION_TEST(
 	FAIAssistantToolsetGetSchema,
@@ -171,11 +309,120 @@ IMPLEMENT_SIMPLE_AUTOMATION_TEST(
 
 bool FAIAssistantToolsetGetSchema::RunTest(const FString& UnusedParameters)
 {
-	CollectGarbage(GARBAGE_COLLECTION_KEEPFLAGS);
 	ToolsetContainer<UFakeToolset> Container;
 	TObjectPtr<UFakeToolset> Instance = *Container;
 
-	TSharedPtr<FJsonObject> ExpectedSchema = GetFakeToolsetSchema(Instance);
+	TSharedPtr<FJsonObject> ExpectedSchema = Instance->GetExpectedSchema();
+
+	FString ExpectedSchemaJson;
+	FJsonSerializer::Serialize(
+		ExpectedSchema.ToSharedRef(),
+		TJsonWriterFactory<TCHAR, TCondensedJsonPrintPolicy<TCHAR>>::Create(&ExpectedSchemaJson));
+	FString OutputSchemaJson = Instance->GetJsonSchema();
+	TestEqual("Toolset schema matches expected", OutputSchemaJson, ExpectedSchemaJson);
+	return true;
+}
+
+// Test that toolset with method with default values can return a schema properly.
+IMPLEMENT_SIMPLE_AUTOMATION_TEST(
+	FAIAssistantToolsetGetSchemaWithDefaultParams,
+	"AI.Assistant.Toolset.GetSchemaWithDefaultParams",
+	AIAssistantTest::Flags)
+
+bool FAIAssistantToolsetGetSchemaWithDefaultParams::RunTest(const FString& UnusedParameters)
+{
+	ToolsetContainer<UFakeToolsetWithDefaultParams> Container;
+	TObjectPtr<UFakeToolsetWithDefaultParams> Instance = *Container;
+
+	TSharedPtr<FJsonObject> ExpectedSchema = Instance->GetExpectedSchema();
+
+	FString ExpectedSchemaJson;
+	FJsonSerializer::Serialize(
+		ExpectedSchema.ToSharedRef(),
+		TJsonWriterFactory<TCHAR, TCondensedJsonPrintPolicy<TCHAR>>::Create(&ExpectedSchemaJson));
+	FString OutputSchemaJson = Instance->GetJsonSchema();
+	TestEqual("Toolset schema matches expected", OutputSchemaJson, ExpectedSchemaJson);
+	return true;
+}
+
+// Test that toolset with no tools can return a schema properly.
+IMPLEMENT_SIMPLE_AUTOMATION_TEST(
+	FAIAssistantToolsetGetSchemaWithNoTools,
+	"AI.Assistant.Toolset.GetSchemaWithNoTools",
+	AIAssistantTest::Flags)
+
+bool FAIAssistantToolsetGetSchemaWithNoTools::RunTest(const FString& UnusedParameters)
+{
+	ToolsetContainer<UFakeToolsetWithNoTools> Container;
+	TObjectPtr<UFakeToolsetWithNoTools> Instance = *Container;
+
+	TSharedPtr<FJsonObject> ExpectedSchema = Instance->GetExpectedSchema();
+
+	FString ExpectedSchemaJson;
+	FJsonSerializer::Serialize(
+		ExpectedSchema.ToSharedRef(),
+		TJsonWriterFactory<TCHAR, TCondensedJsonPrintPolicy<TCHAR>>::Create(&ExpectedSchemaJson));
+	FString OutputSchemaJson = Instance->GetJsonSchema();
+	TestEqual("Toolset schema matches expected", OutputSchemaJson, ExpectedSchemaJson);
+	return true;
+}
+
+// Test that toolset with method with array params can return a schema properly.
+IMPLEMENT_SIMPLE_AUTOMATION_TEST(
+	FAIAssistantToolsetGetSchemaWithArrayParams,
+	"AI.Assistant.Toolset.GetSchemaWithArrayParams",
+	AIAssistantTest::Flags)
+
+bool FAIAssistantToolsetGetSchemaWithArrayParams::RunTest(const FString& UnusedParameters)
+{
+	ToolsetContainer<UFakeToolsetWithArrayParams> Container;
+	TObjectPtr<UFakeToolsetWithArrayParams> Instance = *Container;
+
+	TSharedPtr<FJsonObject> ExpectedSchema = Instance->GetExpectedSchema();
+
+	FString ExpectedSchemaJson;
+	FJsonSerializer::Serialize(
+		ExpectedSchema.ToSharedRef(),
+		TJsonWriterFactory<TCHAR, TCondensedJsonPrintPolicy<TCHAR>>::Create(&ExpectedSchemaJson));
+	FString OutputSchemaJson = Instance->GetJsonSchema();
+	TestEqual("Toolset schema matches expected", OutputSchemaJson, ExpectedSchemaJson);
+	return true;
+}
+
+// Test that toolset with method with map params can return a schema properly.
+IMPLEMENT_SIMPLE_AUTOMATION_TEST(
+	FAIAssistantToolsetGetSchemaWithMapParams,
+	"AI.Assistant.Toolset.GetSchemaWithMapParams",
+	AIAssistantTest::Flags)
+
+bool FAIAssistantToolsetGetSchemaWithMapParams::RunTest(const FString& UnusedParameters)
+{
+	ToolsetContainer<UFakeToolsetWithMapParams> Container;
+	TObjectPtr<UFakeToolsetWithMapParams> Instance = *Container;
+
+	TSharedPtr<FJsonObject> ExpectedSchema = Instance->GetExpectedSchema();
+
+	FString ExpectedSchemaJson;
+	FJsonSerializer::Serialize(
+		ExpectedSchema.ToSharedRef(),
+		TJsonWriterFactory<TCHAR, TCondensedJsonPrintPolicy<TCHAR>>::Create(&ExpectedSchemaJson));
+	FString OutputSchemaJson = Instance->GetJsonSchema();
+	TestEqual("Toolset schema matches expected", OutputSchemaJson, ExpectedSchemaJson);
+	return true;
+}
+
+// Test that toolset with method with set params can return a schema properly.
+IMPLEMENT_SIMPLE_AUTOMATION_TEST(
+	FAIAssistantToolsetGetSchemaWithSetParams,
+	"AI.Assistant.Toolset.GetSchemaWithSetParams",
+	AIAssistantTest::Flags)
+
+bool FAIAssistantToolsetGetSchemaWithSetParams::RunTest(const FString& UnusedParameters)
+{
+	ToolsetContainer<UFakeToolsetWithSetParams> Container;
+	TObjectPtr<UFakeToolsetWithSetParams> Instance = *Container;
+
+	TSharedPtr<FJsonObject> ExpectedSchema = Instance->GetExpectedSchema();
 
 	FString ExpectedSchemaJson;
 	FJsonSerializer::Serialize(
@@ -194,7 +441,6 @@ IMPLEMENT_SIMPLE_AUTOMATION_TEST(
 
 bool FAIAssistantToolsetRegisterUnregisterItself::RunTest(const FString& UnusedParameters)
 {
-	CollectGarbage(GARBAGE_COLLECTION_KEEPFLAGS);
 	ToolsetContainer<UFakeToolset> Container;
 	TObjectPtr<UFakeToolset> Instance = *Container;
 
@@ -217,19 +463,15 @@ IMPLEMENT_SIMPLE_AUTOMATION_TEST(
 
 bool FAIAssistantToolsetExecuteTool::RunTest(const FString& UnusedParameters)
 {
-	CollectGarbage(GARBAGE_COLLECTION_KEEPFLAGS);
 	ToolsetContainer<UFakeToolset> Container;
-
-	FString OutResult;
-	FString OutError;
-	FString InputString = FakeToolsetInput;
+	Container.Input = FakeToolsetInput;
 
 	UAIAssistantToolsetRegistry::ExecuteTool(
-		FakeToolsetName, FakeToolName, InputString, OutResult, OutError);
-	TestTrue("Tool execution returns no error", OutError.IsEmpty());
+		FakeToolsetName, FakeToolName, Container.Input, Container.Result, Container.Error);
+	TestTrue("Tool execution returns no error", Container.Error.IsEmpty());
 
 	FString ExpectedOutput = FString(TEXT(R"({"returnValue":"hello"})"));
-	TestEqual("Tool execution returns expected output", OutResult, ExpectedOutput);
+	TestEqual("Tool execution returns expected output", Container.Result, ExpectedOutput);
 	return true;
 }
 
@@ -241,17 +483,13 @@ IMPLEMENT_SIMPLE_AUTOMATION_TEST(
 
 bool FAIAssistantToolsetExecuteToolInvalidToolset::RunTest(const FString& UnusedParameters)
 {
-	CollectGarbage(GARBAGE_COLLECTION_KEEPFLAGS);
 	ToolsetContainer<UFakeToolset> Container;
-
-	FString OutResult;
-	FString OutError;
-	FString InputString = FakeToolsetInput;
+	Container.Input = FakeToolsetInput;
 
 	AddExpectedError(TEXT("not found"), EAutomationExpectedErrorFlags::Contains);
 	UAIAssistantToolsetRegistry::ExecuteTool(
-		FString("FakeToolsetInvalid"), FakeToolName, InputString, OutResult, OutError);
-	TestFalse("Tool execution with invalid toolset name returns error", OutError.IsEmpty());
+		FString("FakeToolsetInvalid"), FakeToolName, Container.Input, Container.Result, Container.Error);
+	TestFalse("Tool execution with invalid toolset name returns error", Container.Error.IsEmpty());
 	return true;
 }
 
@@ -263,16 +501,12 @@ IMPLEMENT_SIMPLE_AUTOMATION_TEST(
 
 bool FAIAssistantToolsetExecuteToolInvalidToolname::RunTest(const FString& UnusedParameters)
 {
-	CollectGarbage(GARBAGE_COLLECTION_KEEPFLAGS);
 	ToolsetContainer<UFakeToolset> Container;
-
-	FString OutResult;
-	FString OutError;
-	FString InputString = FakeToolsetInput;
+	Container.Input = FakeToolsetInput;
 
 	UAIAssistantToolsetRegistry::ExecuteTool(
-		FakeToolsetName, FString("SimpleToolInvalid"), InputString, OutResult, OutError);
-	TestFalse("Tool execution with invalid tool name returns error", OutError.IsEmpty());
+		FakeToolsetName, FString("SimpleToolInvalid"), Container.Input, Container.Result, Container.Error);
+	TestFalse("Tool execution with invalid tool name returns error", Container.Error.IsEmpty());
 	return true;
 }
 
@@ -284,16 +518,12 @@ IMPLEMENT_SIMPLE_AUTOMATION_TEST(
 
 bool FAIAssistantToolsetExecuteToolExistingButInvalidToolname::RunTest(const FString& UnusedParameters)
 {
-	CollectGarbage(GARBAGE_COLLECTION_KEEPFLAGS);
 	ToolsetContainer<UFakeToolset> Container;
-
-	FString OutResult;
-	FString OutError;
-	FString InputString = FakeToolsetInput;
+	Container.Input = FakeToolsetInput;
 
 	UAIAssistantToolsetRegistry::ExecuteTool(
-		FakeToolsetName, FString("Register"), InputString, OutResult, OutError);
-	TestFalse("Tool execution with invalid tool name returns error", OutError.IsEmpty());
+		FakeToolsetName, FString("Register"), Container.Input, Container.Result, Container.Error);
+	TestFalse("Tool execution with invalid tool name returns error", Container.Error.IsEmpty());
 	return true;
 }
 
@@ -305,20 +535,112 @@ IMPLEMENT_SIMPLE_AUTOMATION_TEST(
 
 bool FAIAssistantToolsetExecuteToolUnregisteredToolset::RunTest(const FString& UnusedParameters)
 {
-	CollectGarbage(GARBAGE_COLLECTION_KEEPFLAGS);
 	ToolsetContainer<UFakeToolset> Container;
-
-	FString OutResult;
-	FString OutError;
-	FString InputString = FakeToolsetInput;
+	Container.Input = FakeToolsetInput;
 
 	TObjectPtr<UFakeToolset> Instance = *Container;
 	Instance->Unregister();
 
 	AddExpectedError(TEXT("not found"), EAutomationExpectedErrorFlags::Contains);
 	UAIAssistantToolsetRegistry::ExecuteTool(
-		FakeToolsetName, FakeToolName, InputString, OutResult, OutError);
-	TestFalse("Tool execution with unregistered toolset returns error", OutError.IsEmpty());
+		FakeToolsetName, FakeToolName, Container.Input, Container.Result, Container.Error);
+	TestFalse("Tool execution with unregistered toolset returns error", Container.Error.IsEmpty());
+	return true;
+}
+
+// Test that querying toolsets thru registry handler work as expected.
+IMPLEMENT_SIMPLE_AUTOMATION_TEST(
+	FAIAssistantToolsetQueryFromRegistryHandler,
+	"AI.Assistant.Toolset.QueryFromRegistryHandler",
+	AIAssistantTest::Flags)
+
+bool FAIAssistantToolsetQueryFromRegistryHandler::RunTest(const FString& UnusedParameters)
+{
+	ToolsetContainer<UFakeToolset> Container;
+	TObjectPtr<UFakeToolset> Instance = *Container;
+
+	auto AssistantSubsystem = UAIAssistantSubsystem::Get(TEXT("Registry handler test"));
+	TestTrue("Can access AI Assistant subsystem", AssistantSubsystem.HasValue());
+
+	TSharedPtr<UE::AIAssistant::IToolsetHandler>* RegisteredHandlerPtr =
+		AssistantSubsystem.GetValue()->ToolsetRegistry.Find(FakeToolsetName);
+	TSharedPtr<UE::AIAssistant::IToolsetHandler> RegisteredHandler =
+			(RegisteredHandlerPtr) ? *RegisteredHandlerPtr : nullptr;
+	TestTrue("Registered toolset handler found", RegisteredHandler.IsValid());
+
+	TestEqual("Toolset registry handler returns name properly",
+			  Instance->GetName(), RegisteredHandler->GetToolsetName());
+	TestEqual("Toolset registry handler returns version properly",
+			  Instance->GetVersion(), RegisteredHandler->GetToolsetVersion());
+	TestEqual("Toolset registry handler returns JSON schema properly",
+			  Instance->GetJsonSchema(), RegisteredHandler->GetToolsetJsonSchema());
+	return true;
+}
+
+// Test that using registry handler to query toolsets that update version in postinit work as expected.
+IMPLEMENT_SIMPLE_AUTOMATION_TEST(
+	FAIAssistantToolsetQueryFromRegistryHandlerUpdatedVersion,
+	"AI.Assistant.Toolset.QueryFromRegistryHandlerUpdatedVersion",
+	AIAssistantTest::Flags)
+
+bool FAIAssistantToolsetQueryFromRegistryHandlerUpdatedVersion::RunTest(const FString& UnusedParameters)
+{
+	ToolsetContainer<UFakeToolsetUpdatesVersion> Container;
+	TObjectPtr<UFakeToolsetUpdatesVersion> Instance = *Container;
+
+	auto AssistantSubsystem = UAIAssistantSubsystem::Get(TEXT("Registry handler test"));
+	TestTrue("Can access AI Assistant subsystem", AssistantSubsystem.HasValue());
+
+	TSharedPtr<UE::AIAssistant::IToolsetHandler>* RegisteredHandlerPtr =
+		AssistantSubsystem.GetValue()->ToolsetRegistry.Find(FakeToolsetUpdatesVersionName);
+	TSharedPtr<UE::AIAssistant::IToolsetHandler> RegisteredHandler =
+			(RegisteredHandlerPtr) ? *RegisteredHandlerPtr : nullptr;
+	TestTrue("Registered toolset handler found", RegisteredHandler.IsValid());
+
+	TestEqual("Toolset registry handler returns name properly",
+			  Instance->GetName(), RegisteredHandler->GetToolsetName());
+	TestEqual("Toolset registry handler returns version properly",
+			  Instance->GetVersion(), RegisteredHandler->GetToolsetVersion());
+	TestEqual("Toolset registry handler returns JSON schema properly",
+			  Instance->GetJsonSchema(), RegisteredHandler->GetToolsetJsonSchema());
+	return true;
+}
+
+// Test that invoking a tool with default values for params provides the proper default value.
+IMPLEMENT_SIMPLE_AUTOMATION_TEST(
+	FAIAssistantToolsetExecuteToolWithDefaultParam,
+	"AI.Assistant.Toolset.ExecuteToolWithDefaultParam",
+	AIAssistantTest::Flags)
+
+bool FAIAssistantToolsetExecuteToolWithDefaultParam::RunTest(const FString& UnusedParameters)
+{
+	ToolsetContainer<UFakeToolsetWithDefaultParams> Container;
+	Container.Input = FakeToolsetInput;
+
+	UAIAssistantToolsetRegistry::ExecuteTool(
+		FString("FakeToolsetWithDefaultParams"), FString("ToolWithDefaults"), Container.Input, Container.Result, Container.Error);
+	TestTrue("Tool execution with default param returns without error", Container.Error.IsEmpty());
+
+	FString ExpectedOutput = FString(TEXT(R"({"returnValue":"hello: 42"})"));
+	TestEqual("Tool execution returns expected output", Container.Result, ExpectedOutput);
+
+	return true;
+}
+
+// Test that invoking a tool with params that don't match input schema errors properly.
+IMPLEMENT_SIMPLE_AUTOMATION_TEST(
+	FAIAssistantToolsetExecuteToolWithBadInput,
+	"AI.Assistant.Toolset.ExecuteToolWithBadInput",
+	AIAssistantTest::Flags)
+
+bool FAIAssistantToolsetExecuteToolWithBadInput::RunTest(const FString& UnusedParameters)
+{
+	ToolsetContainer<UFakeToolset> Container;
+	Container.Input = TEXT(R"({"not_expected":"bad_value"})");
+
+	UAIAssistantToolsetRegistry::ExecuteTool(
+		FakeToolsetName, FakeToolName, Container.Input, Container.Result, Container.Error);
+	TestFalse("Tool execution with invalid param returns with error", Container.Error.IsEmpty());
 	return true;
 }
 
diff --git a/Engine/Plugins/Experimental/AIAssistant/Source/AIAssistant/Private/Tests/AIAssistantToolsetTest.h b/Engine/Plugins/Experimental/AIAssistant/Source/AIAssistant/Private/Tests/AIAssistantToolsetTest.h
--- a/Engine/Plugins/Experimental/AIAssistant/Source/AIAssistant/Private/Tests/AIAssistantToolsetTest.h
+++ b/Engine/Plugins/Experimental/AIAssistant/Source/AIAssistant/Private/Tests/AIAssistantToolsetTest.h
@@ -27,11 +27,13 @@ public:
 	UFUNCTION(meta=(BlueprintInternalUseOnly=true))
 	void SimpleInternalMethod() { }
 
+	TSharedPtr<FJsonObject> GetExpectedSchema();
+
 };
 
 /** Fake Toolset that does nothing except update version. */
 UCLASS(Blueprintable)
-class UFakeToolset2 : public UAIAssistantToolset
+class UFakeToolsetUpdatesVersion : public UAIAssistantToolset
 {
 	GENERATED_BODY()
 
@@ -54,3 +56,86 @@ public:
 	}
 
 };
+
+/** Fake Toolset that has method with default parameters. */
+UCLASS(Blueprintable)
+class UFakeToolsetWithDefaultParams : public UAIAssistantToolset
+{
+	GENERATED_BODY()
+
+public:
+
+	/** Method with some params with default values. */
+	UFUNCTION(BlueprintCallable, Category=Fake)
+	FString ToolWithDefaults(const FString& InString, int InValue=42)
+	{
+		return FString::Printf(TEXT("%s: %d"), *InString, InValue);
+	}
+
+	TSharedPtr<FJsonObject> GetExpectedSchema();
+};
+
+/** Fake Toolset with no tools. */
+UCLASS(Blueprintable)
+class UFakeToolsetWithNoTools : public UAIAssistantToolset
+{
+	GENERATED_BODY()
+
+public:
+	TSharedPtr<FJsonObject> GetExpectedSchema();
+
+};
+
+/** Fake Toolset that has method with TArray parameters. */
+UCLASS(Blueprintable)
+class UFakeToolsetWithArrayParams : public UAIAssistantToolset
+{
+	GENERATED_BODY()
+
+public:
+
+	/** Method with some params with TArray values. */
+	UFUNCTION(BlueprintCallable, Category=Fake)
+	TArray<float> ToolWithArrays(const TArray<int> ArrayParamOne, TArray<FString> ArrayParamTwo)
+	{
+		return TArray<float>();
+	}
+
+	TSharedPtr<FJsonObject> GetExpectedSchema();
+};
+
+/** Fake Toolset that has method with TMap parameters. */
+UCLASS(Blueprintable)
+class UFakeToolsetWithMapParams : public UAIAssistantToolset
+{
+	GENERATED_BODY()
+
+public:
+
+	/** Method with some params with TMap values. */
+	UFUNCTION(BlueprintCallable, Category=Fake)
+	TMap<FString, float> ToolWithMaps(const TMap<FString, int> MapParamOne, TMap<FString, FString> MapParamTwo)
+	{
+		return TMap<FString, float>();
+	}
+
+	TSharedPtr<FJsonObject> GetExpectedSchema();
+};
+
+/** Fake Toolset that has method with TSet parameters. */
+UCLASS(Blueprintable)
+class UFakeToolsetWithSetParams : public UAIAssistantToolset
+{
+	GENERATED_BODY()
+
+public:
+
+	/** Method with some params with TSet values. */
+	UFUNCTION(BlueprintCallable, Category=Fake)
+	TSet<float> ToolWithSets(const TSet<int> SetParamOne, TSet<FString> SetParamTwo)
+	{
+		return TSet<float>();
+	}
+
+	TSharedPtr<FJsonObject> GetExpectedSchema();
+};
