AI Assistant / Tool Calling

More descriptive error messages.

#Synced-CL 49751710

diff --git a/Engine/Plugins/Experimental/AIAssistant/Source/AIAssistant/Private/AIAssistantToolCall.cpp b/Engine/Plugins/Experimental/AIAssistant/Source/AIAssistant/Private/AIAssistantToolCall.cpp
--- a/Engine/Plugins/Experimental/AIAssistant/Source/AIAssistant/Private/AIAssistantToolCall.cpp
+++ b/Engine/Plugins/Experimental/AIAssistant/Source/AIAssistant/Private/AIAssistantToolCall.cpp
@@ -13,8 +13,6 @@
 #include "UObject/FieldIterator.h"
 #include "JsonSchema/JsonSchemaGenerator.h"
 
-#define LOCTEXT_NAMESPACE "FObjectFunctionToolCall"
-
 PRAGMA_DISABLE_DEPRECATION_WARNINGS
 
 using namespace UE::AIAssistant;
@@ -41,6 +39,15 @@ namespace
 		return MakeShared<FJsonObject>();			
 	}
 	
+	FString JsonObjectToString(const TSharedRef<FJsonObject>& JsonObject)
+	{
+		FString JsonString;
+		TSharedRef<TJsonWriter<>> JsonWriter = TJsonWriterFactory<>::Create(&JsonString);
+		FJsonSerializer::Serialize(JsonObject, JsonWriter);
+		
+		return JsonString;
+	}
+
 	TSharedRef<FJsonObject> FunctionInputParamsJsonToValidJsonObject(
 		const TOptional<FObjectFunctionToolCall::FFunctionInputParamsJson>& FunctionInputParamsJson)
 	{
@@ -89,8 +96,7 @@ namespace
 			FName(*Function->GetMetaData(WorldContextKey)) :
 			FName(TEXT("WorldContextObject"));
 		
-		FProperty* WorldContextProperty = Function->FindPropertyByName(WorldContextParamName);
-		if (WorldContextProperty)
+		if (FProperty* WorldContextProperty = Function->FindPropertyByName(WorldContextParamName))
 		{
 			UWorld* World = GetWorld();
 			FObjectPropertyBase* ObjectProperty = 
@@ -105,16 +111,18 @@ namespace
 	/**
 	  * Checks whether all the input parameters required by the function schema JSON are provided
 	  * in the function input params JSON.
+	  * @param FunctionDebugName Name of function, for debugging and error messages.
 	  * @param FunctionInputParamsJsonObject The JSON option containing the function call params.
 	  * @param FunctionSchemaInputSchemaJsonObject The 'inputSchema' sub-JSON in the main function 
 	  *		schema JSON.
-	  *	@param OutErrorText Output error text if a required param is missing, empty otherwise.
+	  *	@param OutErrorString Output error string if a required param is missing, empty otherwise.
 	  *	@return Whether required params are all supplied.
 	  */
 	bool AreAllRequiredFunctionInputParamsProvided(
+		const FString& FunctionDebugName,
 		const TSharedRef<FJsonObject>& FunctionInputParamsJsonObject,
 		const TSharedRef<FJsonObject>& FunctionSchemaInputSchemaJsonObject,
-		FText& OutErrorText)
+		FString* OutErrorString = nullptr)
 	{
 		// For each required param name, make sure we have a field for it in the function input params.
 		{
@@ -126,11 +134,18 @@ namespace
 			{
 				// Caller error. Missing but expected params in incoming function 
 				// input params JSON. 
-				static const FText ErrorText = LOCTEXT(
-					"AIAssistantToolCallError_FunctionInputParamsJsonIsMissingRequiredParams",	
-					"Function input params are required by the function input schema JSON, "
-					"but is incoming function input params JSON is empty.");
-				OutErrorText = ErrorText;
+				if (OutErrorString)
+				{
+					*OutErrorString = FString::Printf(TEXT(
+						"Function \"%s\", "
+						"input params are required by the function input schema Json, "
+						"but incoming function input params Json is empty.\n"
+						"Function schema Json -\n"
+						"%s"),
+						*FunctionDebugName,
+						*JsonObjectToString(FunctionSchemaInputSchemaJsonObject));
+				}
+				
 				return false;
 			}
     		
@@ -139,13 +154,23 @@ namespace
 				if (!FunctionInputParamsJsonObject->HasField(*RequiredParamName))
 				{
 					// Caller error. A missing but expected param in incoming function 
-					// input params JSON. 
-					static const FText ErrorTextTemplate = LOCTEXT(
-						"AIAssistantToolCallError_FunctionInputParamsJsonIsMissingRequiredParam",	
-						"Function input param '{0}' is required by the function input schema JSON, "
-						"but is missing from the incoming function input params JSON.");
-					OutErrorText = FText::Format(
-						ErrorTextTemplate, FText::FromString(RequiredParamName));
+					// input params JSON.
+					if (OutErrorString)
+					{
+						*OutErrorString = FString::Printf(TEXT(
+							"Function \"%s\", "
+							"input param \"%s\" is required by the function input schema Json, "
+							"but is missing from the incoming function input params Json.\n"
+							"Function schema Json -\n"
+							"%s\n"
+							"Function input params Json -\n"
+							"%s"),
+							*FunctionDebugName,
+							*RequiredParamName,
+							*JsonObjectToString(FunctionSchemaInputSchemaJsonObject),
+							*JsonObjectToString(FunctionInputParamsJsonObject));
+					}
+					
 					return false;
 				}
 			}
@@ -155,22 +180,27 @@ namespace
 	}
 
 	/**
-	 * For each input schema property, if we don't have a field for it in the 
-	 * function input params, then add it to the function input params, and 
-	 * give it the same value as the default value of the corresponding function
-	 * input schema property.
-	 * (Note - Checked that we have all the required params first, so we know those 
-	 * will all be in the function input params JSON. Then this ends up only handling
-	 * params that are both non-required and missing from the function input param 
-	 * JSON. To apply default values, we need a 'default' value to exist for each 
-	 * corresponding param object in the function input schema.)
-	  * @param FunctionInputParamsJsonObject The JSON option containing the function call params.
-	  * @param FunctionSchemaInputSchemaJsonObject The 'inputSchema' sub-JSON in the main function 
+	  * For each input schema property, if we don't have a field for it in the 
+	  * function input params, then add it to the function input params, and 
+	  * give it the same value as the default value of the corresponding function
+	  * input schema property.
+	  * (Note - Checked that we have all the required params first, so we know those 
+	  * will all be in the function input params JSON. Then this ends up only handling
+	  * params that are both non-required and missing from the function input param 
+	  * JSON. To apply default values, we need a 'default' value to exist for each 
+ 	  * corresponding param object in the function input schema.)
+ 	  * @param FunctionDebugName Name of the function, used for error message if needed.
+	  * @param FunctionInputParamsJsonObject The JSON object containing the function call params.
+	  * @param FunctionSchemaInputSchemaJsonObject The 'inputSchema' sub-JSON in the main function
 	  *		schema JSON.
+	  * @param OutErrorString Error message, if function fails.
+	  * @return Whether this function could provide defaults for all missing function input params.
 	 */
-	void ProvideDefaultsForMissingFunctionInputParams(
+	bool ProvideDefaultsForMissingFunctionInputParams(
+		const FString& FunctionDebugName,
 		const TSharedRef<FJsonObject>& FunctionInputParamsJsonObject,
-		const TSharedRef<FJsonObject>& FunctionSchemaInputSchemaJsonObject)
+		const TSharedRef<FJsonObject>& FunctionSchemaInputSchemaJsonObject,
+		FString* OutErrorString = nullptr)
 	{
 		if (const TSharedPtr<FJsonObject>* InputPropertiesJsonObject;
 			FunctionSchemaInputSchemaJsonObject->TryGetObjectField(
@@ -182,18 +212,56 @@ namespace
 				const FString& InputParamPropertyName = Pair.Key;
 
 				const TSharedPtr<FJsonObject>& InputPropertyJsonObject = Pair.Value->AsObject();
-				check(InputPropertyJsonObject)
+				if (!InputPropertyJsonObject)
+				{
+					// Caller error. Missing but expected Json for an input param 
+					// in incoming function input params JSON. 
+					if (OutErrorString)
+					{
+						*OutErrorString = FString::Printf(TEXT(
+							"Function \"%s\", "
+							"input param \"%s\" has no Json object field in input params.\n"
+							"Function input params Json -\n"
+							"%s"),
+							*FunctionDebugName,
+							*InputParamPropertyName,
+							*JsonObjectToString(FunctionInputParamsJsonObject));
+					}
+					
+					return false;
+				}
     			
 				if (!FunctionInputParamsJsonObject->HasField(*InputParamPropertyName))
 				{
 					const TSharedPtr<FJsonValue> InputPropertyDefaultJsonValue = 
 						InputPropertyJsonObject->TryGetField(TEXT("default"));
-					check(InputPropertyDefaultJsonValue); // ..this really needs o exist
+					
+					if (!InputPropertyDefaultJsonValue)
+					{
+						// Caller error. Missing but expected Json for an input param's 
+						// default in incoming function input params JSON. 
+						if (OutErrorString)
+						{
+							*OutErrorString = FString::Printf(TEXT(
+								"Function \"%s\", "
+								"input param \"%s\" needs a default value."
+								"Function input params Json -\n"
+								"%s"),
+								*FunctionDebugName,
+								*InputParamPropertyName,
+								*JsonObjectToString(FunctionInputParamsJsonObject));
+						}
+						
+						return false;
+					}
+					
 					FunctionInputParamsJsonObject->SetField(*InputParamPropertyName, 
 						InputPropertyDefaultJsonValue);
 				}
 			}
 		}
+		
+		return true;
 	}
 }
 
@@ -231,7 +299,7 @@ FObjectFunctionToolCall::FToolCallResult FObjectFunctionToolCall::Execute(
 {
 	// Get function input params JSON as a JSON object.
 	TSharedPtr<FJsonObject> FunctionInputParamsJsonObject;
-	if (const TValueOrError<const TSharedRef<FJsonObject>, const FText> JsonObjectOrError = 
+	if (const TValueOrError<const TSharedRef<FJsonObject>, const FString> JsonObjectOrError = 
 			BuildValidFunctionInputParamsJsonObject(FunctionInputParamsJson);
 		JsonObjectOrError.HasError())
 	{
@@ -279,11 +347,14 @@ FObjectFunctionToolCall::FToolCallResult FObjectFunctionToolCall::Execute(
 			FunctionInputParamsJsonObject.ToSharedRef(), Function, FrameMemory))
 	{
 		// Caller error. Ill-formed incoming function input params JSON. 
-		static const FText ErrorText = LOCTEXT(
-			"AIAssistantToolCallError_JsonObjectToUStructFailed", 
-			"Could not convert incoming function input params JSON to a UStruct. "
-			"Incoming function input params JSON may be incorrect or contain unsupported types.");
-		return MakeError(ErrorText);
+		return MakeError(FString::Printf(TEXT(
+			"Function \"%s\", "
+    		"could not convert incoming function input params Json to a UStruct. "
+        	"Incoming function input params Json may be incorrect or contain unsupported types.\n"
+        	"Function input params Json -\n"
+        	"%s"), 
+			*Function->GetName(),
+			*JsonObjectToString(FunctionInputParamsJsonObject.ToSharedRef())));
 	}
 	
 	// In case this is a Blueprint function that has a world context object parameter, 
@@ -307,11 +378,12 @@ FObjectFunctionToolCall::FToolCallResult FObjectFunctionToolCall::Execute(
 		if (!ReturnJsonValue)
 		{
 			// Caller error. Unsupported return value type. 
-			static const FText ErrorText = LOCTEXT(
-				"AIAssistantToolCallError_ReturnValueToJsonValueFailed", 
-				"Could not convert function return value to JSON value. "
-				"Function return value type may be unsupported.");
-			return MakeError(ErrorText);
+			return MakeError(FString::Printf(TEXT(
+				"Function \"%s\", "
+				"could not convert function return value to Json value. "
+				"Function return value type (%s) may be unsupported."),
+				*Function->GetName(),
+				*ReturnProperty->GetClass()->GetName()));
 		}
 		
 		if (ReturnJsonValue->Type == EJson::Object)
@@ -330,12 +402,10 @@ FObjectFunctionToolCall::FToolCallResult FObjectFunctionToolCall::Execute(
 	return MakeValue(ReturnJsonObject);
 }
 
-TValueOrError<const TSharedRef<FJsonObject>, const FText> 
+TValueOrError<const TSharedRef<FJsonObject>, const FString> 
 FObjectFunctionToolCall::BuildValidFunctionInputParamsJsonObject(
 	const TOptional<FFunctionInputParamsJson>& FunctionInputParamsJson) const
 {
-	check(FunctionSchemaJsonObject);
-
 	// Get incoming function params as JSON object.
 	const TSharedRef<FJsonObject> FunctionInputParamsJsonObject = 
 		FunctionInputParamsJsonToValidJsonObject(FunctionInputParamsJson);
@@ -348,26 +418,31 @@ FObjectFunctionToolCall::BuildValidFunctionInputParamsJsonObject(
     	FunctionSchemaJsonObject->TryGetObjectField(
     		TEXT("inputSchema"), FunctionSchemaInputSchemaJsonObjectPtr))
     {
-		// TODO - We should ultimately be validating the function input params JSON against the 
-		// function schema's input schema JSON instead of this. This isn't validating type
-		// compatibility, only checking existence.
-		if (FText ErrorText;
+		// TODO - We should ultimately be fully validating the function input params JSON 
+		// against the function schema's input schema JSON instead of this. This isn't 
+		// validating type compatibility, only checking existence of params.
+		if (FString ErrorString;
 			!AreAllRequiredFunctionInputParamsProvided(
+				Function->GetName(),
 				FunctionInputParamsJsonObject, 
 				FunctionSchemaInputSchemaJsonObjectPtr->ToSharedRef(), 
-				ErrorText))
+				&ErrorString))
 		{
-			return MakeError(ErrorText);
+			return MakeError(ErrorString);
 		}
 
-		ProvideDefaultsForMissingFunctionInputParams(
-			FunctionInputParamsJsonObject, 
-			FunctionSchemaInputSchemaJsonObjectPtr->ToSharedRef()); 
+		if (FString ErrorString;
+			!ProvideDefaultsForMissingFunctionInputParams(
+				Function->GetName(),
+				FunctionInputParamsJsonObject, 
+				FunctionSchemaInputSchemaJsonObjectPtr->ToSharedRef(),
+				&ErrorString))
+		{
+			return MakeError(ErrorString);
+		}
 	}
 	
 	return MakeValue(FunctionInputParamsJsonObject);
 }
 
 PRAGMA_ENABLE_DEPRECATION_WARNINGS
-
-#undef LOCTEXT_NAMESPACE
diff --git a/Engine/Plugins/Experimental/AIAssistant/Source/AIAssistant/Private/AIAssistantToolCall.h b/Engine/Plugins/Experimental/AIAssistant/Source/AIAssistant/Private/AIAssistantToolCall.h
--- a/Engine/Plugins/Experimental/AIAssistant/Source/AIAssistant/Private/AIAssistantToolCall.h
+++ b/Engine/Plugins/Experimental/AIAssistant/Source/AIAssistant/Private/AIAssistantToolCall.h
@@ -27,9 +27,9 @@ namespace UE::AIAssistant
 
 		/**
 		 * Tool call output result.
-		 * Either JSON object or text describing an error.
+		 * Either a JSON object or a string describing an error.
 		 */
-		using FToolCallResult = TValueOrError<const TSharedPtr<FJsonObject>, const FText>;
+		using FToolCallResult = TValueOrError<const TSharedPtr<FJsonObject>, const FString>;
 		
 		/**
 		 * Initializes use of this class.
@@ -39,8 +39,8 @@ namespace UE::AIAssistant
 		
 		/**
 		 * Constructs with a UFunction on a UObject instance. 
-		 * @param InInstanceObject UObject to call the function on.
-		 * @param InFunction UFunction to call on the UObject.
+		* @param InInstanceObject UObject containing the function.
+		 * @param InFunction UFunction contained in the UObject.
 		 */
 		FObjectFunctionToolCall(
 			TNotNull<UObject*> InInstanceObject, TNotNull<UFunction*> InFunction);
@@ -48,8 +48,8 @@ namespace UE::AIAssistant
 		/**
 		 * Constructs with a UFunction on a UObject instance, but using the UObject's UClass 
 		 * where the function can be found by name.
-		 * @param InInstanceObject UObject to call the function on.
-		 * @param InFunctionName UFunction name to find and call on the UObject.
+		 * @param InInstanceObject UObject containing the function.
+		 * @param InFunctionName UFunction name contained in the UObject.
 		 * @return A function call object.
 		 */
 		static FObjectFunctionToolCall ConstructWithNamedClassFunction(
@@ -74,7 +74,7 @@ namespace UE::AIAssistant
 		
 	private:
 
-		TValueOrError<const TSharedRef<FJsonObject>, const FText> 
+		TValueOrError<const TSharedRef<FJsonObject>, const FString> 
 		BuildValidFunctionInputParamsJsonObject(
 			const TOptional<FFunctionInputParamsJson>& FunctionInputParamsJson) const;
 		
diff --git a/Engine/Plugins/Experimental/AIAssistant/Source/AIAssistant/Private/AIAssistantToolset.cpp b/Engine/Plugins/Experimental/AIAssistant/Source/AIAssistant/Private/AIAssistantToolset.cpp
--- a/Engine/Plugins/Experimental/AIAssistant/Source/AIAssistant/Private/AIAssistantToolset.cpp
+++ b/Engine/Plugins/Experimental/AIAssistant/Source/AIAssistant/Private/AIAssistantToolset.cpp
@@ -258,7 +258,7 @@ void UAIAssistantToolset::InvokeTool(
 
 	if (ToolCallResult.HasError())
 	{
-		FString Error = ToolCallResult.StealError().ToString();
+		FString Error = ToolCallResult.StealError();
 		OutError = FString::Printf(TEXT("Error executing tool %s: %s"), *ToolName, *Error);
 		return;
 	}
diff --git a/Engine/Plugins/Experimental/AIAssistant/Source/AIAssistant/Private/Tests/AIAssistantToolCallTest.cpp b/Engine/Plugins/Experimental/AIAssistant/Source/AIAssistant/Private/Tests/AIAssistantToolCallTest.cpp
--- a/Engine/Plugins/Experimental/AIAssistant/Source/AIAssistant/Private/Tests/AIAssistantToolCallTest.cpp
+++ b/Engine/Plugins/Experimental/AIAssistant/Source/AIAssistant/Private/Tests/AIAssistantToolCallTest.cpp
@@ -15,18 +15,25 @@ using namespace UE::AIAssistant;
 
 PRAGMA_DISABLE_DEPRECATION_WARNINGS
 
-#define PRINT_JSONS_TO_LOG 1
+#define DEBUG_PRINT_JSONS_TO_LOG 0
 
-static void MaybePrintJsonToLog(const TCHAR* Name, const TSharedRef<FJsonObject>& JsonObject)
+namespace
 {
-#if PRINT_JSONS_TO_LOG
-	
-	FString JsonString;
-	const TSharedRef<TJsonWriter<>> JsonWriter = TJsonWriterFactory<>::Create(&JsonString);
-	FJsonSerializer::Serialize(JsonObject, JsonWriter);
-	UE_LOG(LogAIAssistant, Warning, TEXT("[%s]\n%s"), Name, *JsonString);
-	
-#endif
+	void PrintJsonToLog(
+		const TCHAR* Name, const TSharedPtr<FJsonObject>& JsonObject)
+	{
+		if (JsonObject)
+		{
+			FString JsonString;
+			const TSharedRef<TJsonWriter<>> JsonWriter = TJsonWriterFactory<>::Create(&JsonString);
+			FJsonSerializer::Serialize(JsonObject, JsonWriter);
+			UE_LOG(LogAIAssistant, Display, TEXT("Json \"%s\" is valid -\n%s"), Name, *JsonString);
+		}
+		else
+		{
+			UE_LOG(LogAIAssistant, Error, TEXT("Json \"%s\" is invalid."), Name);
+		}
+	}
 }
 
 #if WITH_DEV_AUTOMATION_TESTS
@@ -69,7 +76,7 @@ bool FAIAssistantToolCallExecuteToolCallFunctionWithRequiredParam::RunTest(
 			!ToolCallResult.HasError()))
 	{
 		UE_LOG(LogAIAssistant, Error, TEXT("Tool call error - %s"), 
-			*ToolCallResult.GetError().ToString());
+			*ToolCallResult.GetError());
 		return false;
 	}
 	if (!TestTrue("Execute tool call result has value.", 
@@ -86,8 +93,10 @@ bool FAIAssistantToolCallExecuteToolCallFunctionWithRequiredParam::RunTest(
 		return false;
 	}
 
-	MaybePrintJsonToLog(TEXT("ToolCallWithRequiredParamReturnedJsonObject"), 
+#if DEBUG_PRINT_JSONS_TO_LOG
+	PrintJsonToLog(TEXT("ToolCallWithRequiredParamReturnedJsonObject"), 
 		JsonObject.ToSharedRef());
+#endif
 	
 	if (!TestTrue("Return object is not empty (has values.)", 
 			JsonObject->Values.Num() > 0))
@@ -156,7 +165,7 @@ bool FAIAssistantToolCallExecuteToolCallFunctionWithOptionalParam::RunTest(
 			!ToolCallResult.HasError()))
 	{
 		UE_LOG(LogAIAssistant, Error, TEXT("Tool call error - %s"), 
-			*ToolCallResult.GetError().ToString());
+			*ToolCallResult.GetError());
 		return false;
 	}
 	if (!TestTrue("Execute tool call result has value.", 
@@ -173,8 +182,10 @@ bool FAIAssistantToolCallExecuteToolCallFunctionWithOptionalParam::RunTest(
 		return false;
 	}
 
-	MaybePrintJsonToLog(TEXT("ToolCallWithOptionalParamReturnedJsonObject"), 
+#if DEBUG_PRINT_JSONS_TO_LOG
+	PrintJsonToLog(TEXT("ToolCallWithOptionalParamReturnedJsonObject"), 
 		JsonObject.ToSharedRef());
+#endif
 	
 	if (!TestTrue("Return object is not empty (has values.)", 
 			JsonObject->Values.Num() > 0))
