[AI Assistant] Toolsets can now be any Blueprint Function Library

#rnx

#Synced-CL 49840909

diff --git a/Engine/Plugins/Experimental/AIAssistant/Source/AIAssistant/Private/AIAssistantBPToolsetRegistry.cpp b/Engine/Plugins/Experimental/AIAssistant/Source/AIAssistant/Private/AIAssistantBPToolsetRegistry.cpp
new file mode 100644
--- /dev/null
+++ b/Engine/Plugins/Experimental/AIAssistant/Source/AIAssistant/Private/AIAssistantBPToolsetRegistry.cpp
@@ -0,0 +1,150 @@
+// Copyright Epic Games, Inc. All Rights Reserved.
+
+#include "AIAssistantBPToolsetRegistry.h"
+
+#include "Async/Future.h"
+#include "Editor.h"
+#include "Misc/Variant.h"
+#include "Templates/SharedPointer.h"
+
+#include "AIAssistantLog.h"
+
+#include "AIAssistantSubsystem.h"
+#include "AIAssistantToolset.h"
+#include "AIAssistantToolsetHandler.h"
+#include "AIAssistantToolsetRegistry.h"
+
+void UAIAssistantToolsetRegistry::RegisterToolsetClass(UClass* InToolsetClass)
+{
+	if (!InToolsetClass)
+	{
+		UE_LOG(LogAIAssistant, Error, TEXT("Cannot register null toolset class"));
+		return;
+	}
+
+	if (auto AssistantSubsystem = UAIAssistantSubsystem::Get(TEXT("IsToolsetClassRegistered"));
+		AssistantSubsystem.HasValue())
+	{
+		FString ToolsetName = InToolsetClass->GetName();
+		TSharedPtr<UE::AIAssistant::FToolsetLibraryWrapper> Wrapper =
+			MakeShared<UE::AIAssistant::FToolsetLibraryWrapper>(InToolsetClass);
+
+		if (!Wrapper->HasValidTools())
+		{
+			UE_LOG(LogAIAssistant, Warning,
+				   TEXT("Unable to register Toolset %s: invalid tool list"), *ToolsetName);
+			return;
+		}
+
+		if (AssistantSubsystem.GetValue()->ToolsetRegistry.RegisterToolset(Wrapper))
+		{
+			UE_LOG(LogAIAssistant, Display, TEXT("Registering Toolset %s"), *ToolsetName);
+		}
+		else
+		{
+			UE_LOG(LogAIAssistant, Warning, TEXT("Unable to register Toolset %s"), *ToolsetName);
+		}
+	}
+	else
+	{
+		UE_LOG(LogAIAssistant, Error, TEXT("AIAssistantSubsystem unavailable"));
+	}
+}
+
+void UAIAssistantToolsetRegistry::UnregisterToolsetClass(UClass* InToolsetClass)
+{
+	if (!InToolsetClass)
+	{
+		UE_LOG(LogAIAssistant, Warning, TEXT("Attempting to unregister null toolset class"));
+		return;
+	}
+
+	if (auto AssistantSubsystem = UAIAssistantSubsystem::Get(TEXT("IsToolsetClassRegistered"));
+		AssistantSubsystem.HasValue())
+	{
+		FString ToolsetName = InToolsetClass->GetName();
+		TSharedPtr<UE::AIAssistant::IToolsetHandler>* PtrToHandler = AssistantSubsystem.GetValue()->ToolsetRegistry.Find(ToolsetName);
+		if (PtrToHandler)
+		{
+			TSharedPtr<UE::AIAssistant::IToolsetHandler> Handler = *PtrToHandler;
+			UE_LOG(LogAIAssistant, Display, TEXT("Unregistering Toolset %s"), *ToolsetName);
+			AssistantSubsystem.GetValue()->ToolsetRegistry.UnregisterToolset(*PtrToHandler);
+		}
+	}
+	else
+	{
+		UE_LOG(LogAIAssistant, Error, TEXT("AIAssistantSubsystem unavailable"));
+	}
+}
+
+bool UAIAssistantToolsetRegistry::IsToolsetClassRegistered(UClass* InToolsetClass)
+{
+	if (!InToolsetClass)
+	{
+		return false;
+	}
+
+	if (auto AssistantSubsystem = UAIAssistantSubsystem::Get(TEXT("IsToolsetClassRegistered"));
+		AssistantSubsystem.HasValue())
+	{
+		FString ToolsetName = InToolsetClass->GetName();
+		return AssistantSubsystem.GetValue()->ToolsetRegistry.Find(ToolsetName) != nullptr;
+	}
+	return false;
+}
+
+bool UAIAssistantToolsetRegistry::IsToolsetRegistered(const FString& InToolsetName)
+{
+	if (auto AssistantSubsystem = UAIAssistantSubsystem::Get(TEXT("IsToolsetClassRegistered"));
+		AssistantSubsystem.HasValue())
+	{
+		return AssistantSubsystem.GetValue()->ToolsetRegistry.Find(InToolsetName) != nullptr;
+	}
+	return false;
+}
+
+void UAIAssistantToolsetRegistry::ExecuteTool(
+	const FString& ToolsetName, const FString& ToolName, const FString& JsonInput, FString& OutResult, FString& OutError)
+{
+	OutResult.Empty();
+	OutError.Empty();
+
+	auto AssistantSubsystem = UAIAssistantSubsystem::Get(TEXT("ExecuteTool"));
+	if (AssistantSubsystem.HasError())
+	{
+		OutError = AssistantSubsystem.StealError();
+		return;
+	}
+
+	if (AssistantSubsystem.HasValue())
+	{
+		UE::AIAssistant::FToolDescriptor ToolDescriptor;
+		ToolDescriptor.ToolsetName = ToolsetName;
+		ToolDescriptor.ToolName = ToolName;
+		TFuture<TValueOrError<FString, FString>> Future =
+			AssistantSubsystem.GetValue()->ToolsetRegistry.ExecuteTool(ToolDescriptor, JsonInput);
+		// TODO: Block for now, will do proper async handling soon.
+		TValueOrError<FString, FString> Result = Future.Get();
+		if (Result.HasError())
+		{
+			OutError = Result.StealError();
+		}
+		else if (Result.HasValue())
+		{
+			OutResult = Result.StealValue();
+		}
+	}
+}
+
+FString UAIAssistantToolsetRegistry::GetToolsetJsonSchema(UClass* InToolsetClass)
+{
+	if (!InToolsetClass)
+	{
+		UE_LOG(LogAIAssistant, Warning, TEXT("Attempting to get schema of null toolset class"));
+		return FString();
+	}
+
+	TSharedPtr<UE::AIAssistant::FToolsetLibraryWrapper> Wrapper =
+		MakeShared<UE::AIAssistant::FToolsetLibraryWrapper>(InToolsetClass);
+	return Wrapper ? Wrapper->GetToolsetJsonSchema() : FString();
+}
diff --git a/Engine/Plugins/Experimental/AIAssistant/Source/AIAssistant/Private/AIAssistantBPToolsetRegistry.h b/Engine/Plugins/Experimental/AIAssistant/Source/AIAssistant/Private/AIAssistantBPToolsetRegistry.h
--- a/Engine/Plugins/Experimental/AIAssistant/Source/AIAssistant/Private/AIAssistantBPToolsetRegistry.h
+++ b/Engine/Plugins/Experimental/AIAssistant/Source/AIAssistant/Private/AIAssistantBPToolsetRegistry.h
@@ -3,16 +3,7 @@
 #pragma once
 
 #include "CoreMinimal.h"
-#include "Async/Future.h"
 #include "Kismet/BlueprintFunctionLibrary.h"
-#include "Editor.h"
-#include "Misc/Variant.h"
-#include "Templates/SharedPointer.h"
-
-#include "AIAssistantLog.h"
-#include "AIAssistantSubsystem.h"
-#include "AIAssistantToolsetRegistry.h"
-#include "AIAssistantToolsetHandler.h"
 
 #include "AIAssistantBPToolsetRegistry.generated.h"
 
@@ -24,62 +15,30 @@ class AIASSISTANT_API UAIAssistantToolsetRegistry : public UBlueprintFunctionLib
 
 public:
 
-	UFUNCTION(BlueprintPure, Category="ToolsetRegistry")
-	static bool IsToolsetRegistered(const FString& ToolsetName)
-	{
-		auto AssistantSubsystem = UAIAssistantSubsystem::Get(TEXT("IsToolsetRegistered"));
-		return AssistantSubsystem.HasValue() &&
-			static_cast<bool>(AssistantSubsystem.GetValue()->ToolsetRegistry.Find(ToolsetName));
-	}
-
+	// Register a Blueprint function library class as a toolset.
 	UFUNCTION(BlueprintCallable, Category="ToolsetRegistry")
-	static void UnregisterToolsetByName(const FString& ToolsetName)
-	{
-		auto AssistantSubsystem = UAIAssistantSubsystem::Get(TEXT("UnregisterToolsetByName"));
-		if (AssistantSubsystem.HasValue())
-		{
-			TSharedPtr<UE::AIAssistant::IToolsetHandler>* PtrToHandler =
-				AssistantSubsystem.GetValue()->ToolsetRegistry.Find(ToolsetName);
-			if (PtrToHandler)
-			{
-				TSharedPtr<UE::AIAssistant::IToolsetHandler> Handler = *PtrToHandler;
-				UE_LOG(LogAIAssistant, Display, TEXT("Unregistering Toolset %s"), *ToolsetName);
-				AssistantSubsystem.GetValue()->ToolsetRegistry.UnregisterToolset(Handler);
-			}
-		}
-	}
+	static void RegisterToolsetClass(UClass* InToolsetClass);
 
+	// Unregister a Blueprint function library class as a toolset.
 	UFUNCTION(BlueprintCallable, Category="ToolsetRegistry")
-	static void ExecuteTool(const FString& ToolsetName, const FString& ToolName, const FString& JsonInput,
-							FString& OutResult, FString& OutError)
-	{
-		OutResult.Empty();
-		OutError.Empty();
+	static void UnregisterToolsetClass(UClass* InToolsetClass);
 
-		auto AssistantSubsystem = UAIAssistantSubsystem::Get(TEXT("ExecuteTool"));
-		if (AssistantSubsystem.HasError())
-		{
-			OutError = AssistantSubsystem.StealError();
-			return;
-		}
+ 	// Return whether a toolset class is already registered.
+	UFUNCTION(BlueprintPure, Category="ToolsetRegistry")
+	static bool IsToolsetClassRegistered(UClass* InToolsetClass);
 
-		if (AssistantSubsystem.HasValue())
-		{
-			UE::AIAssistant::FToolDescriptor ToolDescriptor;
-			ToolDescriptor.ToolsetName = ToolsetName;
-			ToolDescriptor.ToolName = ToolName;
-			TFuture<TValueOrError<FString, FString>> Future =
-				AssistantSubsystem.GetValue()->ToolsetRegistry.ExecuteTool(ToolDescriptor, JsonInput);
-			// TODO: Block for now, will do proper async handling soon.
-			TValueOrError<FString, FString> Result = Future.Get();
-			if (Result.HasError())
-			{
-				OutError = Result.StealError();
-			}
-			else if (Result.HasValue())
-			{
-				OutResult = Result.StealValue();
-			}
-		}
-	}
+ 	// Return whether a toolset is already registered by name.
+	UFUNCTION(BlueprintPure, Category="ToolsetRegistry")
+	static bool IsToolsetRegistered(const FString& InToolsetName);
+
+	// Execute a registered tool given a toolset and tool name as if it was being executed by the Assistant.
+	// Meant for use while testing.
+	UFUNCTION(BlueprintCallable, Category="ToolsetRegistry")
+	static void ExecuteTool(
+		const FString& ToolsetName, const FString& ToolName, const FString& JsonInput, FString& OutResult, FString& OutError);
+
+	// Get toolset JSON schema for a toolset class.
+	// Meant for use while testing.
+	UFUNCTION(BlueprintPure, Category="ToolsetRegistry")
+	static FString GetToolsetJsonSchema(UClass* InToolsetClass);
 };
diff --git a/Engine/Plugins/Experimental/AIAssistant/Source/AIAssistant/Private/AIAssistantToolset.cpp b/Engine/Plugins/Experimental/AIAssistant/Source/AIAssistant/Private/AIAssistantToolset.cpp
--- a/Engine/Plugins/Experimental/AIAssistant/Source/AIAssistant/Private/AIAssistantToolset.cpp
+++ b/Engine/Plugins/Experimental/AIAssistant/Source/AIAssistant/Private/AIAssistantToolset.cpp
@@ -5,325 +5,196 @@
 #include "Async/Future.h"
 #include "Containers/UnrealString.h"
 #include "Dom/JsonObject.h"
-#include "Editor.h"
 #include "JsonSchema/JsonSchemaGenerator.h"
+#include "Kismet/BlueprintFunctionLibrary.h"
 #include "Policies/CondensedJsonPrintPolicy.h"
 #include "Serialization/JsonWriter.h"
 #include "Serialization/JsonSerializer.h"
-#include "Templates/ValueOrError.h"
-#include "UObject/Class.h"
 
 #include "AIAssistantLog.h"
 #include "AIAssistantSubsystem.h"
 #include "AIAssistantToolCall.h"
-#include "AIAssistantToolsetHandler.h"
-
 
 namespace UE::AIAssistant
 {
-
-	/**
-	 * Implements IToolsetHandler for registration.
-	 *
-	 * Hands off calls from AI Assistant registry to UAIAssistantToolset.
-	 */
-	class FToolsetHandlerTrampoline : public UE::AIAssistant::IToolsetHandler
+	TFuture<TValueOrError<FString, FString>> FToolsetLibraryWrapper::ExecuteTool(
+		const FString& ToolName, const FString& JsonInput)
 	{
-	  public:
-		FToolsetHandlerTrampoline(TWeakObjectPtr<UAIAssistantToolset> InToolsetHandler)
-			: ToolsetHandler(InToolsetHandler)
+		FString ToolResult;
+		FString ToolError;
+		InvokeTool(ToolName, JsonInput, ToolResult, ToolError);
+		if (ToolError.IsEmpty())
 		{
-			TStrongObjectPtr<UAIAssistantToolset> Handler = ToolsetHandler.Pin();
-			ToolsetName = Handler ? Handler->GetName() : FString();
+			return MakeFulfilledPromise<TValueOrError<FString, FString>>(
+				MakeValue(ToolResult)).GetFuture();
 		}
-
-		// Executes a tool with the given name and JSON input, returning a future with the JSON
-		// output.
-		virtual TFuture<TValueOrError<FString, FString>> ExecuteTool(
-			const FString& ToolName, const FString& JsonInput) override
+		else
 		{
-			FString ToolResult;
-			FString ToolError;
-			TStrongObjectPtr<UAIAssistantToolset> Handler = ToolsetHandler.Pin();
-			if (Handler)
-			{
-				Handler->InvokeTool(ToolName, JsonInput, ToolResult, ToolError);
-				if (ToolError.IsEmpty())
-				{
-					return MakeFulfilledPromise<TValueOrError<FString, FString>>(
-						MakeValue(ToolResult)).GetFuture();
-				}
-				else
-				{
-					return MakeFulfilledPromise<TValueOrError<FString, FString>>(
-						MakeError(ToolError)).GetFuture();
-				}
-			}
-			else
-			{
-				return MakeFulfilledPromise<TValueOrError<FString, FString>>(
-					MakeError(FString(TEXT("Invalid Toolset Handler")))).GetFuture();
-			}
+			return MakeFulfilledPromise<TValueOrError<FString, FString>>(
+				MakeError(ToolError)).GetFuture();
 		}
-
-		// Returns the JSON schema for the toolset
-		virtual FString GetToolsetJsonSchema() const override
-		{
-			TStrongObjectPtr<UAIAssistantToolset> Handler = ToolsetHandler.Pin();
-			return Handler ? Handler->GetJsonSchema() : FString();
-		}
-
-		// Returns the name of the toolset
-		virtual FString GetToolsetName() const override
-		{
-			return ToolsetName;
-		}
-
-		// Returns the version of the toolset
-		virtual FString GetToolsetVersion() const override
-		{
-			TStrongObjectPtr<UAIAssistantToolset> Handler = ToolsetHandler.Pin();
-			return Handler ? Handler->GetVersion(): FString();
-		}
-
-	  private:
-		TWeakObjectPtr<UAIAssistantToolset> ToolsetHandler;
-
-		// Caching the toolset name is required in order to be able to Unregister from the registry when the
-		// owning UAIAssistantToolset is garbage collected, because the registry uses the toolset name as keys
-		// into its map of handlers.  Cannot ask the owning UAIAssistantToolset for its name at that point, as
-		// it's being destroyed.
-		FString ToolsetName;
-	};
-}
-
-void UAIAssistantToolset::PostInitProperties()
-{
-	Super::PostInitProperties();
-
-	// Only actual instances should register.
-	if (!HasAnyFlags(RF_ClassDefaultObject | RF_ArchetypeObject))
-	{
-		Register();
 	}
-}
 
-void UAIAssistantToolset::BeginDestroy()
-{
-	// Only actual instances should unregister.
-	if (!HasAnyFlags(RF_ClassDefaultObject | RF_ArchetypeObject))
+	FString FToolsetLibraryWrapper::GetToolsetJsonSchema() const
 	{
-		Unregister();
-	}
-	Super::BeginDestroy();
-}
+		FString OutputString;
+		TSharedPtr<FJsonObject> Schema = GenerateSchema();
 
-FString UAIAssistantToolset::GetJsonSchema() const
-{
-	FString OutputString;
-	TSharedPtr<FJsonObject> Schema = GenerateSchema();
+		if (!FJsonSerializer::Serialize(
+				Schema.ToSharedRef(),
+				TJsonWriterFactory<TCHAR, TCondensedJsonPrintPolicy<TCHAR>>::Create(&OutputString)))
+		{
+			UE_LOG(
+				LogAIAssistant, Warning,
+				TEXT("Failed to generate valid toolset schema for toolset %s"), *GetToolsetName());
+		}
 
-	if (!FJsonSerializer::Serialize(
-			Schema.ToSharedRef(),
-			TJsonWriterFactory<TCHAR, TCondensedJsonPrintPolicy<TCHAR>>::Create(&OutputString)))
-	{
-		UE_LOG(
-			LogAIAssistant, Warning,
-			TEXT("Failed to generate valid toolset schema for toolset %s"), *GetName());
+		return OutputString;
 	}
 
-	return OutputString;
-}
-
-TSharedPtr<FJsonObject> UAIAssistantToolset::GenerateSchema() const
-{
-	TSharedPtr<FJsonObject> Schema = MakeShared<FJsonObject>();
-	check(Schema);
-	Schema->SetStringField(FString(TEXT("name")), GetName());
-	Schema->SetStringField(FString(TEXT("version")), GetVersion());
-	Schema->SetStringField(FString(TEXT("description")), GetDescription());
-
-	TArray<TSharedPtr<FJsonValue>> ToolsArray;
-	TArray<TObjectPtr<UFunction>> ValidTools = GetValidToolMethodsList();
-	for (UFunction* Function : ValidTools)
+	TSharedPtr<FJsonObject> FToolsetLibraryWrapper::GenerateSchema() const
 	{
-        check(Function);
+		TSharedPtr<FJsonObject> Schema = MakeShared<FJsonObject>();
+		check(Schema);
+		Schema->SetStringField(FString(TEXT("name")), GetToolsetName());
+		Schema->SetStringField(FString(TEXT("version")), GetToolsetVersion());
+		Schema->SetStringField(FString(TEXT("description")), GetToolsetDescription());
 
-PRAGMA_DISABLE_DEPRECATION_WARNINGS
-		TSharedPtr<FJsonObject> ToolEntry = FJsonSchemaGenerator::UStructToJsonSchemaObject(Function);
-PRAGMA_ENABLE_DEPRECATION_WARNINGS
+		TArray<TSharedPtr<FJsonValue>> ToolsArray;
+		TArray<TObjectPtr<UFunction>> ValidTools = GetValidToolMethodsList();
+		for (UFunction* Function : ValidTools)
+		{
+			check(Function);
 
-		ToolsArray.Add(MakeShared<FJsonValueObject>(ToolEntry));
-	}
-	Schema->SetArrayField(FString(TEXT("tools")), ToolsArray);
-	return Schema;
-}
+			PRAGMA_DISABLE_DEPRECATION_WARNINGS
+				TSharedPtr<FJsonObject> ToolEntry = FJsonSchemaGenerator::UStructToJsonSchemaObject(Function);
+			PRAGMA_ENABLE_DEPRECATION_WARNINGS
 
-void UAIAssistantToolset::GenerateToolCallObjects()
-{
-	Tools.Empty();
-	TArray<TObjectPtr<UFunction>> ValidTools = GetValidToolMethodsList();
-	for (UFunction* Function : ValidTools)
-	{
-		check(Function);
-		const UE::AIAssistant::FObjectFunctionToolCall ToolCall =
-			UE::AIAssistant::FObjectFunctionToolCall(
-				TNotNull<UObject*>(this), TNotNull<UFunction*>(Function));
-		Tools.Add(
-			Function->GetName(), MakeShared<UE::AIAssistant::FObjectFunctionToolCall>(ToolCall));
+			ToolsArray.Add(MakeShared<FJsonValueObject>(ToolEntry));
+		}
+		Schema->SetArrayField(FString(TEXT("tools")), ToolsArray);
+		return Schema;
 	}
-}
 
-void UAIAssistantToolset::Register()
-{
-	auto AssistantSubsystem = UAIAssistantSubsystem::Get(TEXT("Register"));
-	if (AssistantSubsystem.HasValue())
+	void FToolsetLibraryWrapper::GenerateToolCallObjects()
 	{
-		FString ToolsetName = GetName();
-		if (AssistantSubsystem.GetValue()->ToolsetRegistry.Find(ToolsetName))
+		if (!Tools.IsEmpty())
 		{
-			UE_LOG(LogAIAssistant, Warning,
-				   TEXT("Toolset %s is already registered"), *ToolsetName);
-			return;
+			Tools.Empty();
 		}
-
-		if (!HasValidTools())
+		TStrongObjectPtr<UClass> Toolset = ToolsetClass.Pin();
+		if (!Toolset)
 		{
-			UE_LOG(LogAIAssistant, Warning,
-				   TEXT("Unable to register Toolset %s: invalid tool list"), *ToolsetName);
 			return;
 		}
 
-		if (!RegistryHandler.IsValid())
+		TArray<TObjectPtr<UFunction>> ValidTools = GetValidToolMethodsList();
+		for (UFunction* Function : ValidTools)
 		{
-			RegistryHandler = MakeShared<UE::AIAssistant::FToolsetHandlerTrampoline>(
-				TWeakObjectPtr<UAIAssistantToolset>(this));
+			check(Function);
+			const FObjectFunctionToolCall ToolCall =
+				FObjectFunctionToolCall(TNotNull<UObject*>(Toolset.Get()), TNotNull<UFunction*>(Function));
+			Tools.Add(
+				Function->GetName(), MakeShared<FObjectFunctionToolCall>(ToolCall));
 		}
-
-		UE_LOG(LogAIAssistant, Display, TEXT("Registering Toolset %s"), *ToolsetName);
-		AssistantSubsystem.GetValue()->ToolsetRegistry.RegisterToolset(RegistryHandler);
 	}
-}
 
-void UAIAssistantToolset::Unregister()
-{
-	auto AssistantSubsystem = UAIAssistantSubsystem::Get(TEXT("Unregister"));
-	if (AssistantSubsystem.HasValue())
+	void FToolsetLibraryWrapper::InvokeTool(const FString& ToolName, const FString& JsonInput, FString& OutResult, FString& OutError)
 	{
-		FString ToolsetName = GetName();
-		TSharedPtr<UE::AIAssistant::IToolsetHandler>* PtrToRegisteredHandler =
-			AssistantSubsystem.GetValue()->ToolsetRegistry.Find(ToolsetName);
-		TSharedPtr<UE::AIAssistant::IToolsetHandler> RegisteredHandler =
-			(PtrToRegisteredHandler) ? *PtrToRegisteredHandler : nullptr;
-		if (!RegisteredHandler)
+		OutResult.Empty();
+		OutError.Empty();
+		TStrongObjectPtr<UClass> Toolset = ToolsetClass.Pin();
+		if (!Toolset)
 		{
-			// Nothing is registered, nothing to do.
+			OutError = FString(TEXT("Invalid Toolset Class"));
 			return;
 		}
 
-		if (RegistryHandler.IsValid() && RegistryHandler == RegisteredHandler)
+		TSharedPtr<FObjectFunctionToolCall>* ToolCallPtr = Tools.Find(ToolName);
+		if (!ToolCallPtr)
 		{
-			UE_LOG(LogAIAssistant, Display, TEXT("Unregistering Toolset %s"), *ToolsetName);
-			AssistantSubsystem.GetValue()->ToolsetRegistry.UnregisterToolset(RegistryHandler);
+			OutError = FString::Printf(TEXT("Unknown tool %s"), *ToolName);
+			return;
 		}
-	}
-}
+		TSharedPtr<FObjectFunctionToolCall> ToolCall = *ToolCallPtr;
+		check(ToolCall);
 
-void UAIAssistantToolset::InvokeTool(
-	const FString& ToolName, const FString& JsonInput, FString& OutResult, FString& OutError)
-{
-	OutResult.Empty();
-	OutError.Empty();
+		TSharedPtr<FJsonObject> ToolCallReturnValue;
+		FObjectFunctionToolCall::FFunctionInputParamsJson JsonParams;
+		JsonParams.Set<FString>(JsonInput);
 
-	// Ensure valid tool list is up to date
-	if (Tools.IsEmpty())
-	{
-		GenerateToolCallObjects();
-	}
-
-	TSharedPtr<UE::AIAssistant::FObjectFunctionToolCall>* ToolCallPtr = Tools.Find(ToolName);
-	if (!ToolCallPtr)
-	{
-		OutError = FString::Printf(TEXT("Unknown tool %s"), *ToolName);
-		return;
-	}
-	TSharedPtr<UE::AIAssistant::FObjectFunctionToolCall> ToolCall = *ToolCallPtr;
-	check(ToolCall);
+		FObjectFunctionToolCall::FToolCallResult ToolCallResult = ToolCall->Execute(JsonParams);
 
-	TSharedPtr<FJsonObject> ToolCallReturnValue;
-	UE::AIAssistant::FObjectFunctionToolCall::FFunctionInputParamsJson JsonParams;
-	JsonParams.Set<FString>(JsonInput);
+		if (ToolCallResult.HasError())
+		{
+			FString Error = ToolCallResult.StealError();
+			OutError = FString::Printf(TEXT("Error executing tool %s: %s"), *ToolName, *Error);
+			return;
+		}
 
-	UE::AIAssistant::FObjectFunctionToolCall::FToolCallResult ToolCallResult = ToolCall->Execute(JsonParams);
+		if (!ToolCallResult.HasValue())
+		{
+			OutError = FString::Printf(TEXT("No return value from executing tool %s"), *ToolName);
+			return;
+		}
 
-	if (ToolCallResult.HasError())
-	{
-		FString Error = ToolCallResult.StealError();
-		OutError = FString::Printf(TEXT("Error executing tool %s: %s"), *ToolName, *Error);
-		return;
-	}
+		ToolCallReturnValue = ToolCallResult.StealValue();
+		if (!ToolCallReturnValue.IsValid())
+		{
+			OutError = FString::Printf(TEXT("No return value from executing tool %s"), *ToolName);
+			return;
+		}
 
-	if (!ToolCallResult.HasValue())
-	{
-		OutError = FString::Printf(TEXT("No return value from executing tool %s"), *ToolName);
-		return;
-	}
+		if (!FJsonSerializer::Serialize(
+				ToolCallReturnValue.ToSharedRef(),
+				TJsonWriterFactory<TCHAR, TCondensedJsonPrintPolicy<TCHAR>>::Create(&OutResult)))
+		{
+			OutError = FString::Printf(TEXT("Unable to serialize result from executing tool %s"), *ToolName);
+			return;
+		}
 
-	ToolCallReturnValue = ToolCallResult.StealValue();
-	if (!ToolCallReturnValue.IsValid())
-	{
-		OutError = FString::Printf(TEXT("No return value from executing tool %s"), *ToolName);
 		return;
 	}
 
-	if (!FJsonSerializer::Serialize(
-			ToolCallReturnValue.ToSharedRef(),
-			TJsonWriterFactory<TCHAR, TCondensedJsonPrintPolicy<TCHAR>>::Create(&OutResult)))
+	TArray<FString> FToolsetLibraryWrapper::GetToolNames()
 	{
-		OutError = FString::Printf(TEXT("Unable to serialize result from executing tool %s"), *ToolName);
-		return;
+		TArray<TObjectPtr<UFunction>> ValidTools = GetValidToolMethodsList();
+		TArray<FString> ToolNames;
+		for (TObjectPtr<UFunction> FuncPtr : ValidTools)
+		{
+			ToolNames.Add(FuncPtr->GetName());
+		}
+		return ToolNames;
 	}
 
-	return;
-}
-
-TArray<FString> UAIAssistantToolset::GetToolNames()
-{
-	TArray<TObjectPtr<UFunction>> ValidTools = GetValidToolMethodsList();
-	TArray<FString> ToolNames;
-	for (TObjectPtr<UFunction> FuncPtr : ValidTools)
+	bool FToolsetLibraryWrapper::HasValidTools()
 	{
-		ToolNames.Add(FuncPtr->GetName());
+		return !Tools.IsEmpty();
 	}
-	return ToolNames;
-}
 
-bool UAIAssistantToolset::HasValidTools()
-{
-	GenerateToolCallObjects();
-	if (Tools.IsEmpty())
+	TArray<TObjectPtr<UFunction>> FToolsetLibraryWrapper::GetValidToolMethodsList() const
 	{
-		return false;
-	}
-	// TODO: when we have validation methods for tool schemas we can check them.
-	return true;
-}
+		TArray<TObjectPtr<UFunction>> ValidTools;
+		TStrongObjectPtr<UClass> Toolset = ToolsetClass.Pin();
+		if (!Toolset)
+		{
+			return ValidTools;
+		}
 
-TArray<TObjectPtr<UFunction>> UAIAssistantToolset::GetValidToolMethodsList() const
-{
-	TArray<TObjectPtr<UFunction>> ValidTools;
-	for (TFieldIterator<UFunction> FuncIt(GetClass(), EFieldIterationFlags::None); FuncIt; ++FuncIt)
-	{
-		TObjectPtr<UFunction> Function = *FuncIt;
-        if (Function)
+		for (TFieldIterator<UFunction> FuncIt(Toolset.Get(), EFieldIterationFlags::None); FuncIt; ++FuncIt)
 		{
-			// Skip functions marked as for internal use only.
-			if (Function->HasMetaData(TEXT("BlueprintInternalUseOnly")))
+			TObjectPtr<UFunction> Function = *FuncIt;
+			if (Function)
 			{
-				continue;
+				// Skip functions marked as for internal use only.
+				if (Function->HasMetaData(TEXT("BlueprintInternalUseOnly")))
+				{
+					continue;
+				}
+
+				ValidTools.Add(Function);
 			}
-			ValidTools.Add(Function);
 		}
+		return ValidTools;
 	}
-	return ValidTools;
 }
diff --git a/Engine/Plugins/Experimental/AIAssistant/Source/AIAssistant/Private/AIAssistantToolset.h b/Engine/Plugins/Experimental/AIAssistant/Source/AIAssistant/Private/AIAssistantToolset.h
--- a/Engine/Plugins/Experimental/AIAssistant/Source/AIAssistant/Private/AIAssistantToolset.h
+++ b/Engine/Plugins/Experimental/AIAssistant/Source/AIAssistant/Private/AIAssistantToolset.h
@@ -4,109 +4,83 @@
 
 #include "CoreMinimal.h"
 
-#include "Containers/UnrealString.h"
 #include "Dom/JsonObject.h"
+#include "Kismet/BlueprintFunctionLibrary.h"
+#include "Templates/SharedPointer.h"
 
-#include "AIAssistantToolCall.h"
 #include "AIAssistantToolsetHandler.h"
 
-#include "AIAssistantToolset.generated.h"
 
 namespace UE::AIAssistant
 {
-	class FToolsetHandlerTrampoline;
-}
-
-/**
- * Provide AI Assistant toolset handler implementation.
- *
- */
-UCLASS(Blueprintable)
-class AIASSISTANT_API UAIAssistantToolset: public UObject
-{
-	GENERATED_BODY()
-
-public:
+	class FObjectFunctionToolCall;
 
-	static inline FString DEFAULT_TOOLSET_VERSION = FString("1.0");
-
-	// Returns the JSON schema for the toolset
-	UFUNCTION(BlueprintCallable, Category=AssistantToolset)
-	FString GetJsonSchema() const;
-
-	// Returns the name of the toolset
-	UFUNCTION(BlueprintCallable, Category=AssistantToolset)
-	FString GetName() const
+	// Wraps a Blueprint Function Library for use as an AIAssistant Toolset.
+	class FToolsetLibraryWrapper : public IToolsetHandler
 	{
-		return GetClass()->GetName();
-	}
-
-	// Returns the version of the toolset
-	UFUNCTION(BlueprintCallable, Category=AssistantToolset)
-	FString GetVersion() const
-	{
-		return ToolsetVersion;
-	}
-
-	// Returns the description of the toolset
-	UFUNCTION(BlueprintCallable, Category=AssistantToolset)
-	FString GetDescription() const
-	{
-		return GetClass()->GetToolTipText().ToString();
-	}
-
-	// Registers the toolset with the AI Assistant.
-	UFUNCTION(BlueprintCallable, Category=AssistantToolset)
-	void Register();
-
-	// Unregisters the toolset from the AI Assistant.
-	UFUNCTION(BlueprintCallable, Category=AssistantToolset)
-	void Unregister();
-
-	// Provides toolset version to the AI Assistant.
-	// Override this in PostInitProperties.
-	UPROPERTY(BlueprintReadOnly, Category=AssistantToolset)
-	FString ToolsetVersion = DEFAULT_TOOLSET_VERSION;
-
-	// Invoke a tool in this toolset by name.
-	//
-	// Expected to be called by the toolset registry ExecuteTool method, so
-	// ToolName should be the unqualified function name.
-	//
-	// Returns JSON result in OutResult unless there was an error.
-	//
-	// Returns error message in OutError if there was error in invocation.
-	UFUNCTION(BlueprintCallable, Category=AssistantToolset)
-	void InvokeTool(const FString& ToolName, const FString& JsonInput, FString& OutResult, FString& OutError);
-
-	// Return list of tool names in this toolset.
-	UFUNCTION(BlueprintCallable, Category=AssistantToolset)
-	TArray<FString> GetToolNames();
-
-	// Overridden to automatically register toolset with AI Assistant subsystem.
-	// Subclasses should also override the ToolsetVersion property here.
-	virtual void PostInitProperties() override;
-	// Overridden to automatically unregister toolset with AI Assistant subsystem.
-	virtual void BeginDestroy() override;
-
-private:
-
-	// Handler that implements IToolsetHandler to register with the AI Assistant.
-	TSharedPtr<UE::AIAssistant::FToolsetHandlerTrampoline> RegistryHandler;
-
-	// Map of valid tool names to tool call objects.
-	TMap<FString, TSharedPtr<UE::AIAssistant::FObjectFunctionToolCall>> Tools;
-
-	// Whether the toolset has a valid set of tools.
-	bool HasValidTools();
-
-	// Generate schema from reflection data.
-	TSharedPtr<FJsonObject> GenerateSchema() const;
-
-	// Generate Tools list.
-	void GenerateToolCallObjects();
-
-	// Helper method to generate list of valid tools.
-	TArray<TObjectPtr<UFunction>> GetValidToolMethodsList() const;
-};
-
+	  public:
+		FToolsetLibraryWrapper(UClass* LibraryClass) : ToolsetClass(LibraryClass)
+		{
+			GenerateToolCallObjects();
+		}
+
+		// Executes a tool with the given name and JSON input, returning a future with the JSON
+		// output.
+		virtual TFuture<TValueOrError<FString, FString>> ExecuteTool(
+			const FString& ToolName, const FString& JsonInput) override;
+
+		// Returns the JSON schema for the toolset
+		virtual FString GetToolsetJsonSchema() const override;
+
+		// Returns the name of the toolset
+		virtual FString GetToolsetName() const override
+		{
+			return ToolsetClass.Pin() ? ToolsetClass->GetName() : FString();
+		}
+
+		// Returns the version of the toolset
+		virtual FString GetToolsetVersion() const override
+		{
+			// TODO: Investigate some way for toolset to provide a version string.
+			return FString("1.0");
+		}
+
+		// Returns the description of the toolset from the class's metadata.
+		FString GetToolsetDescription() const
+		{
+			return ToolsetClass.Pin() ? ToolsetClass->GetToolTipText().ToString() : FString();
+		}
+
+		// Invoke a tool in this toolset by name.
+		//
+		// Expected to be called by the toolset registry ExecuteTool method, so
+		// ToolName should be the unqualified function name.
+		//
+		// Returns JSON result in OutResult unless there was an error.
+		//
+		// Returns error message in OutError if there was error in invocation.
+		void InvokeTool(const FString& ToolName, const FString& JsonInput, FString& OutResult, FString& OutError);
+
+		TArray<FString> GetToolNames();
+
+		// Whether the toolset has a valid set of tools.
+		bool HasValidTools();
+
+	  private:
+
+		// The wrapped toolset function library
+		TWeakObjectPtr<UClass> ToolsetClass;
+
+		// Map of valid tool names to tool call objects.
+		TMap<FString, TSharedPtr<FObjectFunctionToolCall>> Tools;
+
+		// Generate schema from reflection data.
+		TSharedPtr<FJsonObject> GenerateSchema() const;
+
+		// Generate Tools list.
+		void GenerateToolCallObjects();
+
+		// Helper method to generate list of valid tools.
+		TArray<TObjectPtr<UFunction>> GetValidToolMethodsList() const;
+	};
+}
diff --git a/Engine/Plugins/Experimental/AIAssistant/Source/AIAssistant/Private/Tests/AIAssistantToolsetTest.cpp b/Engine/Plugins/Experimental/AIAssistant/Source/AIAssistant/Private/Tests/AIAssistantToolsetTest.cpp
--- a/Engine/Plugins/Experimental/AIAssistant/Source/AIAssistant/Private/Tests/AIAssistantToolsetTest.cpp
+++ b/Engine/Plugins/Experimental/AIAssistant/Source/AIAssistant/Private/Tests/AIAssistantToolsetTest.cpp
@@ -13,53 +13,59 @@
 #include "Serialization/JsonWriter.h"
 #include "Serialization/JsonSerializer.h"
 
-#include "AIAssistantBPToolsetRegistry.h"
 #include "AIAssistantLog.h"
+#include "AIAssistantBPToolsetRegistry.h"
+#include "AIAssistantSubsystem.h"
+#include "AIAssistantToolset.h"
 #include "Tests/AIAssistantTestSchemaBuilder.h"
 #include "Tests/AIAssistantTestFlags.h"
 
 
 #if WITH_DEV_AUTOMATION_TESTS
-
 const FString FakeToolsetName(TEXT("FakeToolset"));
 const FString FakeToolsetUpdatesVersionName(TEXT("FakeToolsetUpdatesVersion"));
 const FString FakeToolsetInput(TEXT(R"({"inString":"hello"})"));
 const FString FakeToolName(TEXT("SimpleTool"));
 const FString ReturnValueName(TEXT("returnValue"));
 
-// Wrapper for creating and cleaning up UObjects.
-template<typename T>class ToolsetContainer
+class ToolsetRegistryContainer
 {
-public:
-	ToolsetContainer()
+  public:
+	ToolsetRegistryContainer(UClass* InToolsetClass)
+		: ToolsetClass(InToolsetClass)
 	{
-		CollectGarbage(GARBAGE_COLLECTION_KEEPFLAGS);
-		Instance = NewObject<T>();
+		UAIAssistantToolsetRegistry::RegisterToolsetClass(ToolsetClass);
 	}
 
-	~ToolsetContainer()
+	~ToolsetRegistryContainer()
 	{
-		Instance->MarkAsGarbage();
-		Instance = nullptr;
-		CollectGarbage(GARBAGE_COLLECTION_KEEPFLAGS);
+		UAIAssistantToolsetRegistry::UnregisterToolsetClass(ToolsetClass);
 	}
 
-	TObjectPtr<T> operator*() const
+	TSharedPtr<UE::AIAssistant::FToolsetLibraryWrapper> ToolsetHandler()
 	{
-		return Instance;
+		auto AssistantSubsystem = UAIAssistantSubsystem::Get(TEXT("Test Register"));
+		if (AssistantSubsystem.HasValue() && ToolsetClass)
+		{
+			FString ToolsetName = ToolsetClass->GetName();
+
+			TSharedPtr<UE::AIAssistant::IToolsetHandler>* PtrToHandler =
+				AssistantSubsystem.GetValue()->ToolsetRegistry.Find(ToolsetName);
+			if (PtrToHandler)
+			{
+				return StaticCastSharedPtr<UE::AIAssistant::FToolsetLibraryWrapper>(*PtrToHandler);
+			}
+		}
+		return nullptr;
 	}
 
-	FString Input;
-	FString Result;
-	FString Error;
-
-private:
-	TObjectPtr<T> Instance;
+	UClass* ToolsetClass;
 };
 
-TSharedPtr<FJsonObject> UFakeToolset::GetExpectedSchema()
+TSharedPtr<FJsonObject> GetFakeToolsetExpectedSchema()
 {
-	FJsonDomBuilder::FObject SchemaBuilder = MakeSchemaBuilder(GetName(), GetVersion(), GetDescription());
+	FJsonDomBuilder::FObject SchemaBuilder = MakeSchemaBuilder(
+		FString(TEXT("FakeToolset")), FString("1.0"), FString(TEXT("Fake Toolset that does nothing.")));
 
 	FJsonDomBuilder::FObject InputSchema = MakeSchema();
 	FJsonDomBuilder::FObject Properties;
@@ -84,9 +90,11 @@ TSharedPtr<FJsonObject> UFakeToolset::GetExpectedSchema()
 	return SchemaBuilder.AsJsonObject();
 }
 
-TSharedPtr<FJsonObject> UFakeToolsetWithDefaultParams::GetExpectedSchema()
+TSharedPtr<FJsonObject> GetFakeToolsetWithDefaultParamsExpectedSchema()
 {
-	FJsonDomBuilder::FObject SchemaBuilder = MakeSchemaBuilder(GetName(), GetVersion(), GetDescription());
+	FJsonDomBuilder::FObject SchemaBuilder = MakeSchemaBuilder(
+		FString(TEXT("FakeToolsetWithDefaultParams")), FString("1.0"),
+		FString(TEXT("Fake Toolset that has method with default parameters.")));
 
 	FJsonDomBuilder::FObject InputSchema = MakeSchema();
 	FJsonDomBuilder::FObject Properties;
@@ -114,18 +122,11 @@ TSharedPtr<FJsonObject> UFakeToolsetWithDefaultParams::GetExpectedSchema()
 	return SchemaBuilder.AsJsonObject();
 }
 
-TSharedPtr<FJsonObject> UFakeToolsetWithNoTools::GetExpectedSchema()
-{
-	FJsonDomBuilder::FObject SchemaBuilder = MakeSchemaBuilder(GetName(), GetVersion(), GetDescription());
-
-	FJsonDomBuilder::FArray ToolsArray;
-	SchemaBuilder.Set(TEXT("tools"), ToolsArray);
-	return SchemaBuilder.AsJsonObject();
-}
-
-TSharedPtr<FJsonObject> UFakeToolsetWithArrayParams::GetExpectedSchema()
+TSharedPtr<FJsonObject> GetFakeToolsetWithArrayParamsExpectedSchema()
 {
-	FJsonDomBuilder::FObject SchemaBuilder = MakeSchemaBuilder(GetName(), GetVersion(), GetDescription());
+	FJsonDomBuilder::FObject SchemaBuilder = MakeSchemaBuilder(
+		FString(TEXT("FakeToolsetWithArrayParams")), FString("1.0"),
+		FString(TEXT("Fake Toolset that has method with TArray parameters.")));
 
 	FJsonDomBuilder::FObject InputSchema = MakeSchema();
 
@@ -155,9 +156,11 @@ TSharedPtr<FJsonObject> UFakeToolsetWithArrayParams::GetExpectedSchema()
 	return SchemaBuilder.AsJsonObject();
 }
 
-TSharedPtr<FJsonObject> UFakeToolsetWithSetParams::GetExpectedSchema()
+TSharedPtr<FJsonObject> GetFakeToolsetWithSetParamsExpectedSchema()
 {
-	FJsonDomBuilder::FObject SchemaBuilder = MakeSchemaBuilder(GetName(), GetVersion(), GetDescription());
+	FJsonDomBuilder::FObject SchemaBuilder = MakeSchemaBuilder(
+		FString(TEXT("FakeToolsetWithSetParams")), FString("1.0"),
+		FString(TEXT("Fake Toolset that has method with TSet parameters.")));
 
 	FJsonDomBuilder::FObject InputSchema = MakeSchema();
 
@@ -187,9 +190,11 @@ TSharedPtr<FJsonObject> UFakeToolsetWithSetParams::GetExpectedSchema()
 	return SchemaBuilder.AsJsonObject();
 }
 
-TSharedPtr<FJsonObject> UFakeToolsetWithMapParams::GetExpectedSchema()
+TSharedPtr<FJsonObject> GetFakeToolsetWithMapParamsExpectedSchema()
 {
-	FJsonDomBuilder::FObject SchemaBuilder = MakeSchemaBuilder(GetName(), GetVersion(), GetDescription());
+	FJsonDomBuilder::FObject SchemaBuilder = MakeSchemaBuilder(
+		FString(TEXT("FakeToolsetWithMapParams")), FString("1.0"),
+		FString(TEXT("Fake Toolset that has method with TMap parameters.")));
 
 	FJsonDomBuilder::FObject InputSchema = MakeSchema();
 
@@ -219,28 +224,64 @@ TSharedPtr<FJsonObject> UFakeToolsetWithMapParams::GetExpectedSchema()
 	return SchemaBuilder.AsJsonObject();
 }
 
-// Test that toolset class creation registers toolset and destroy unregisters toolset.
+// Test that registering with null toolset class input fails properly.
 IMPLEMENT_SIMPLE_AUTOMATION_TEST(
-	FAIAssistantToolsetRegistersOnCreateAndDestroy,
-	"AI.Assistant.Toolset.RegistersOnCreateAndDestroy",
+	FAIAssistantToolsetRegisterNullClass,
+	"AI.Assistant.Toolset.RegisterNullClass",
 	AIAssistantTest::Flags)
 
-bool FAIAssistantToolsetRegistersOnCreateAndDestroy::RunTest(const FString& UnusedParameters)
+bool FAIAssistantToolsetRegisterNullClass::RunTest(const FString& UnusedParameters)
 {
-	// Ensure clean registry.
-	CollectGarbage(GARBAGE_COLLECTION_KEEPFLAGS);
-	UAIAssistantToolsetRegistry::UnregisterToolsetByName(FakeToolsetName);
+	AddExpectedError(TEXT("Cannot register null toolset class"), EAutomationExpectedErrorFlags::Contains);
+	ToolsetRegistryContainer Container(nullptr);
+	TestTrue("Null toolset class", Container.ToolsetClass == nullptr);
+
+	TSharedPtr<UE::AIAssistant::FToolsetLibraryWrapper> Toolset = Container.ToolsetHandler();
+	TestTrue("Null toolset wrapper", Toolset == nullptr);
+
+	return true;
+}
+
+// Test that wrapper with null toolset class input fails properly.
+IMPLEMENT_SIMPLE_AUTOMATION_TEST(
+	FAIAssistantToolsetWrapperNullClass,
+	"AI.Assistant.Toolset.WrapperNullClass",
+	AIAssistantTest::Flags)
+
+bool FAIAssistantToolsetWrapperNullClass::RunTest(const FString& UnusedParameters)
+{
+	UE::AIAssistant::FToolsetLibraryWrapper Wrapper(nullptr);
+	TestFalse("Invalid toolset wrapper with null toolset", Wrapper.HasValidTools());
+	return true;
+}
+
+// Test that toolset can unregister and re-register itself properly.
+IMPLEMENT_SIMPLE_AUTOMATION_TEST(
+	FAIAssistantToolsetRegisterUnregisterToolset,
+	"AI.Assistant.Toolset.RegisterUnregisterToolset",
+	AIAssistantTest::Flags)
 
-	TestFalse("No toolsets registered to start.",
-			  UAIAssistantToolsetRegistry::IsToolsetRegistered(FakeToolsetName));
+bool FAIAssistantToolsetRegisterUnregisterToolset::RunTest(const FString& UnusedParameters)
+{
+	UClass* ToolsetClass = UFakeToolset::StaticClass();
+
+	TestFalse("Toolset unregistered before registering itself.",
+			  UAIAssistantToolsetRegistry::IsToolsetClassRegistered(ToolsetClass));
+
+	// Scope for registration of ToolsetClass.
 	{
-		ToolsetContainer<UFakeToolset> Container;
-		TestTrue("Toolset registered after class instance created.",
-				 UAIAssistantToolsetRegistry::IsToolsetRegistered(FakeToolsetName));
-	}
-	TestFalse("Toolset is unregistered after class instance destroyed.",
-			  UAIAssistantToolsetRegistry::IsToolsetRegistered(FakeToolsetName));
+		ToolsetRegistryContainer Container(ToolsetClass);
+		TestTrue("Toolset registered after registering itself.",
+				 UAIAssistantToolsetRegistry::IsToolsetClassRegistered(ToolsetClass));
 
+		UAIAssistantToolsetRegistry::UnregisterToolsetClass(ToolsetClass);
+		TestFalse("Toolset unregistered after unregistering itself.",
+				  UAIAssistantToolsetRegistry::IsToolsetClassRegistered(ToolsetClass));
+
+		UAIAssistantToolsetRegistry::RegisterToolsetClass(ToolsetClass);
+		TestTrue("Toolset registered after re-registering itself.",
+				 UAIAssistantToolsetRegistry::IsToolsetClassRegistered(ToolsetClass));
+	}
 	return true;
 }
 
@@ -252,13 +293,11 @@ IMPLEMENT_SIMPLE_AUTOMATION_TEST(
 
 bool FAIAssistantToolsetNoToolsRegistrationRejected::RunTest(const FString& UnusedParameters)
 {
-	ToolsetContainer<UFakeToolsetWithNoTools> Container;
-	TObjectPtr<UFakeToolsetWithNoTools> Instance = *Container;
+	UClass* ToolsetClass = UFakeToolsetWithNoTools::StaticClass();
+
+	ToolsetRegistryContainer Container(ToolsetClass);
 	TestFalse("Toolset with no tools is not registered.",
-			  UAIAssistantToolsetRegistry::IsToolsetRegistered(Instance->GetName()));
-	Instance->Register();
-	TestFalse("Toolset with no tools is not registered when asking explicitly.",
-			  UAIAssistantToolsetRegistry::IsToolsetRegistered(Instance->GetName()));
+			  UAIAssistantToolsetRegistry::IsToolsetClassRegistered(ToolsetClass));
 	return true;
 }
 
@@ -270,38 +309,20 @@ IMPLEMENT_SIMPLE_AUTOMATION_TEST(
 
 bool FAIAssistantToolsetCanGetValues::RunTest(const FString& UnusedParameters)
 {
-	ToolsetContainer<UFakeToolset> Container;
-	TObjectPtr<UFakeToolset> Instance = *Container;
-	TestEqual("Toolset returns its own name",
-			  Instance->GetName(), FakeToolsetName);
-	TestEqual("Toolset returns default non-overridden version",
-			  Instance->GetVersion(), UAIAssistantToolset::DEFAULT_TOOLSET_VERSION);
-	TestEqual("Toolset returns description from class definition",
-			  Instance->GetDescription(), FString(TEXT("Fake Toolset that does nothing.")));
-	return true;
-}
-
-// Test that overriding toolset properties for version works properly.
-IMPLEMENT_SIMPLE_AUTOMATION_TEST(
-	FAIAssistantToolsetCanOverrideProperties,
-	"AI.Assistant.Toolset.CanOverrideProperties",
-	AIAssistantTest::Flags)
+	ToolsetRegistryContainer Container(UFakeToolset::StaticClass());
+	TSharedPtr<UE::AIAssistant::FToolsetLibraryWrapper> Toolset = Container.ToolsetHandler();
+	TestTrue("Valid toolset handler returned", Toolset.IsValid());
 
-bool FAIAssistantToolsetCanOverrideProperties::RunTest(const FString& UnusedParameters)
-{
-	ToolsetContainer<UFakeToolsetUpdatesVersion> Container;
-	TObjectPtr<UFakeToolsetUpdatesVersion> Instance = *Container;
 	TestEqual("Toolset returns its own name",
-			  Instance->GetName(), FakeToolsetUpdatesVersionName);
-	TestEqual("Toolset returns overridden version",
-			  Instance->GetVersion(), UFakeToolsetUpdatesVersion::MY_VERSION);
+			  Toolset->GetToolsetName(), FakeToolsetName);
+	TestEqual("Toolset returns default non-overridden version",
+			  Toolset->GetToolsetVersion(), FString("1.0"));
 	TestEqual("Toolset returns description from class definition",
-			  Instance->GetDescription(),
-			  FString(TEXT("Fake Toolset that does nothing except update version.")));
+			  Toolset->GetToolsetDescription(), FString(TEXT("Fake Toolset that does nothing.")));
 	return true;
 }
 
-// Test that toolset can return a schema properly.
+// Test that toolset can return a basic schema properly.
 IMPLEMENT_SIMPLE_AUTOMATION_TEST(
 	FAIAssistantToolsetGetSchema,
 	"AI.Assistant.Toolset.GetSchema",
@@ -309,16 +330,17 @@ IMPLEMENT_SIMPLE_AUTOMATION_TEST(
 
 bool FAIAssistantToolsetGetSchema::RunTest(const FString& UnusedParameters)
 {
-	ToolsetContainer<UFakeToolset> Container;
-	TObjectPtr<UFakeToolset> Instance = *Container;
+	ToolsetRegistryContainer Container(UFakeToolset::StaticClass());
+	TSharedPtr<UE::AIAssistant::FToolsetLibraryWrapper> Toolset = Container.ToolsetHandler();
+	TestTrue("Valid toolset handler returned", Toolset.IsValid());
 
-	TSharedPtr<FJsonObject> ExpectedSchema = Instance->GetExpectedSchema();
+	TSharedPtr<FJsonObject> ExpectedSchema = GetFakeToolsetExpectedSchema();
 
 	FString ExpectedSchemaJson;
 	FJsonSerializer::Serialize(
 		ExpectedSchema.ToSharedRef(),
 		TJsonWriterFactory<TCHAR, TCondensedJsonPrintPolicy<TCHAR>>::Create(&ExpectedSchemaJson));
-	FString OutputSchemaJson = Instance->GetJsonSchema();
+	FString OutputSchemaJson = Toolset->GetToolsetJsonSchema();
 	TestEqual("Toolset schema matches expected", OutputSchemaJson, ExpectedSchemaJson);
 	return true;
 }
@@ -331,38 +353,17 @@ IMPLEMENT_SIMPLE_AUTOMATION_TEST(
 
 bool FAIAssistantToolsetGetSchemaWithDefaultParams::RunTest(const FString& UnusedParameters)
 {
-	ToolsetContainer<UFakeToolsetWithDefaultParams> Container;
-	TObjectPtr<UFakeToolsetWithDefaultParams> Instance = *Container;
-
-	TSharedPtr<FJsonObject> ExpectedSchema = Instance->GetExpectedSchema();
-
-	FString ExpectedSchemaJson;
-	FJsonSerializer::Serialize(
-		ExpectedSchema.ToSharedRef(),
-		TJsonWriterFactory<TCHAR, TCondensedJsonPrintPolicy<TCHAR>>::Create(&ExpectedSchemaJson));
-	FString OutputSchemaJson = Instance->GetJsonSchema();
-	TestEqual("Toolset schema matches expected", OutputSchemaJson, ExpectedSchemaJson);
-	return true;
-}
-
-// Test that toolset with no tools can return a schema properly.
-IMPLEMENT_SIMPLE_AUTOMATION_TEST(
-	FAIAssistantToolsetGetSchemaWithNoTools,
-	"AI.Assistant.Toolset.GetSchemaWithNoTools",
-	AIAssistantTest::Flags)
-
-bool FAIAssistantToolsetGetSchemaWithNoTools::RunTest(const FString& UnusedParameters)
-{
-	ToolsetContainer<UFakeToolsetWithNoTools> Container;
-	TObjectPtr<UFakeToolsetWithNoTools> Instance = *Container;
+	ToolsetRegistryContainer Container(UFakeToolsetWithDefaultParams::StaticClass());
+	TSharedPtr<UE::AIAssistant::FToolsetLibraryWrapper> Toolset = Container.ToolsetHandler();
+	TestTrue("Valid toolset handler returned", Toolset.IsValid());
 
-	TSharedPtr<FJsonObject> ExpectedSchema = Instance->GetExpectedSchema();
+	TSharedPtr<FJsonObject> ExpectedSchema = GetFakeToolsetWithDefaultParamsExpectedSchema();
 
 	FString ExpectedSchemaJson;
 	FJsonSerializer::Serialize(
 		ExpectedSchema.ToSharedRef(),
 		TJsonWriterFactory<TCHAR, TCondensedJsonPrintPolicy<TCHAR>>::Create(&ExpectedSchemaJson));
-	FString OutputSchemaJson = Instance->GetJsonSchema();
+	FString OutputSchemaJson = Toolset->GetToolsetJsonSchema();
 	TestEqual("Toolset schema matches expected", OutputSchemaJson, ExpectedSchemaJson);
 	return true;
 }
@@ -375,16 +376,17 @@ IMPLEMENT_SIMPLE_AUTOMATION_TEST(
 
 bool FAIAssistantToolsetGetSchemaWithArrayParams::RunTest(const FString& UnusedParameters)
 {
-	ToolsetContainer<UFakeToolsetWithArrayParams> Container;
-	TObjectPtr<UFakeToolsetWithArrayParams> Instance = *Container;
+	ToolsetRegistryContainer Container(UFakeToolsetWithArrayParams::StaticClass());
+	TSharedPtr<UE::AIAssistant::FToolsetLibraryWrapper> Toolset = Container.ToolsetHandler();
+	TestTrue("Valid toolset handler returned", Toolset.IsValid());
 
-	TSharedPtr<FJsonObject> ExpectedSchema = Instance->GetExpectedSchema();
+	TSharedPtr<FJsonObject> ExpectedSchema = GetFakeToolsetWithArrayParamsExpectedSchema();
 
 	FString ExpectedSchemaJson;
 	FJsonSerializer::Serialize(
 		ExpectedSchema.ToSharedRef(),
 		TJsonWriterFactory<TCHAR, TCondensedJsonPrintPolicy<TCHAR>>::Create(&ExpectedSchemaJson));
-	FString OutputSchemaJson = Instance->GetJsonSchema();
+	FString OutputSchemaJson = Toolset->GetToolsetJsonSchema();
 	TestEqual("Toolset schema matches expected", OutputSchemaJson, ExpectedSchemaJson);
 	return true;
 }
@@ -397,16 +399,17 @@ IMPLEMENT_SIMPLE_AUTOMATION_TEST(
 
 bool FAIAssistantToolsetGetSchemaWithMapParams::RunTest(const FString& UnusedParameters)
 {
-	ToolsetContainer<UFakeToolsetWithMapParams> Container;
-	TObjectPtr<UFakeToolsetWithMapParams> Instance = *Container;
+	ToolsetRegistryContainer Container(UFakeToolsetWithMapParams::StaticClass());
+	TSharedPtr<UE::AIAssistant::FToolsetLibraryWrapper> Toolset = Container.ToolsetHandler();
+	TestTrue("Valid toolset handler returned", Toolset.IsValid());
 
-	TSharedPtr<FJsonObject> ExpectedSchema = Instance->GetExpectedSchema();
+	TSharedPtr<FJsonObject> ExpectedSchema = GetFakeToolsetWithMapParamsExpectedSchema();
 
 	FString ExpectedSchemaJson;
 	FJsonSerializer::Serialize(
 		ExpectedSchema.ToSharedRef(),
 		TJsonWriterFactory<TCHAR, TCondensedJsonPrintPolicy<TCHAR>>::Create(&ExpectedSchemaJson));
-	FString OutputSchemaJson = Instance->GetJsonSchema();
+	FString OutputSchemaJson = Toolset->GetToolsetJsonSchema();
 	TestEqual("Toolset schema matches expected", OutputSchemaJson, ExpectedSchemaJson);
 	return true;
 }
@@ -419,42 +422,21 @@ IMPLEMENT_SIMPLE_AUTOMATION_TEST(
 
 bool FAIAssistantToolsetGetSchemaWithSetParams::RunTest(const FString& UnusedParameters)
 {
-	ToolsetContainer<UFakeToolsetWithSetParams> Container;
-	TObjectPtr<UFakeToolsetWithSetParams> Instance = *Container;
+	ToolsetRegistryContainer Container(UFakeToolsetWithSetParams::StaticClass());
+	TSharedPtr<UE::AIAssistant::FToolsetLibraryWrapper> Toolset = Container.ToolsetHandler();
+	TestTrue("Valid toolset handler returned", Toolset.IsValid());
 
-	TSharedPtr<FJsonObject> ExpectedSchema = Instance->GetExpectedSchema();
+	TSharedPtr<FJsonObject> ExpectedSchema = GetFakeToolsetWithSetParamsExpectedSchema();
 
 	FString ExpectedSchemaJson;
 	FJsonSerializer::Serialize(
 		ExpectedSchema.ToSharedRef(),
 		TJsonWriterFactory<TCHAR, TCondensedJsonPrintPolicy<TCHAR>>::Create(&ExpectedSchemaJson));
-	FString OutputSchemaJson = Instance->GetJsonSchema();
+	FString OutputSchemaJson = Toolset->GetToolsetJsonSchema();
 	TestEqual("Toolset schema matches expected", OutputSchemaJson, ExpectedSchemaJson);
 	return true;
 }
 
-// Test that toolset can unregister and re-register itself properly.
-IMPLEMENT_SIMPLE_AUTOMATION_TEST(
-	FAIAssistantToolsetRegisterUnregisterItself,
-	"AI.Assistant.Toolset.RegisterUnregisterItself",
-	AIAssistantTest::Flags)
-
-bool FAIAssistantToolsetRegisterUnregisterItself::RunTest(const FString& UnusedParameters)
-{
-	ToolsetContainer<UFakeToolset> Container;
-	TObjectPtr<UFakeToolset> Instance = *Container;
-
-	TestTrue("Toolset registered after class instance created.",
-			  UAIAssistantToolsetRegistry::IsToolsetRegistered(FakeToolsetName));
-	Instance->Unregister();
-	TestFalse("Toolset unregistered after unregistering itself.",
-			  UAIAssistantToolsetRegistry::IsToolsetRegistered(FakeToolsetName));
-	Instance->Register();
-	TestTrue("Toolset registered after re-registering itself.",
-			  UAIAssistantToolsetRegistry::IsToolsetRegistered(FakeToolsetName));
-	return true;
-}
-
 // Test that invoking a tool from a toolset calls the right method and returns the expected result.
 IMPLEMENT_SIMPLE_AUTOMATION_TEST(
 	FAIAssistantToolsetExecuteTool,
@@ -463,33 +445,35 @@ IMPLEMENT_SIMPLE_AUTOMATION_TEST(
 
 bool FAIAssistantToolsetExecuteTool::RunTest(const FString& UnusedParameters)
 {
-	ToolsetContainer<UFakeToolset> Container;
-	Container.Input = FakeToolsetInput;
+	ToolsetRegistryContainer Container(UFakeToolset::StaticClass());
+	TSharedPtr<UE::AIAssistant::FToolsetLibraryWrapper> Toolset = Container.ToolsetHandler();
+	TestTrue("Valid toolset handler returned", Toolset.IsValid());
 
-	UAIAssistantToolsetRegistry::ExecuteTool(
-		FakeToolsetName, FakeToolName, Container.Input, Container.Result, Container.Error);
-	TestTrue("Tool execution returns no error", Container.Error.IsEmpty());
+	FString Result;
+	FString Error;
+
+	Toolset->InvokeTool(FakeToolName, FakeToolsetInput, Result, Error);
+	TestTrue("Tool execution returns no error", Error.IsEmpty());
 
 	FString ExpectedOutput = FString(TEXT(R"({"returnValue":"hello"})"));
-	TestEqual("Tool execution returns expected output", Container.Result, ExpectedOutput);
+	TestEqual("Tool execution returns expected output", Result, ExpectedOutput);
 	return true;
 }
 
-// Test that invoking a tool from an invalid toolset errors out properly.
+// Test that invoking a tool from an invalid named toolset errors out properly.
 IMPLEMENT_SIMPLE_AUTOMATION_TEST(
-	FAIAssistantToolsetExecuteToolInvalidToolset,
-	"AI.Assistant.Toolset.ExecuteToolInvalidToolset",
+	FAIAssistantToolsetExecuteToolInvalidToolsetName,
+	"AI.Assistant.Toolset.ExecuteToolInvalidToolsetName",
 	AIAssistantTest::Flags)
 
-bool FAIAssistantToolsetExecuteToolInvalidToolset::RunTest(const FString& UnusedParameters)
+bool FAIAssistantToolsetExecuteToolInvalidToolsetName::RunTest(const FString& UnusedParameters)
 {
-	ToolsetContainer<UFakeToolset> Container;
-	Container.Input = FakeToolsetInput;
-
+	FString Result;
+	FString Error;
 	AddExpectedError(TEXT("not found"), EAutomationExpectedErrorFlags::Contains);
 	UAIAssistantToolsetRegistry::ExecuteTool(
-		FString("FakeToolsetInvalid"), FakeToolName, Container.Input, Container.Result, Container.Error);
-	TestFalse("Tool execution with invalid toolset name returns error", Container.Error.IsEmpty());
+		FString("FakeToolsetInvalid"), FakeToolName, FakeToolsetInput, Result, Error);
+	TestFalse("Tool execution with invalid toolset name returns error", Error.IsEmpty());
 	return true;
 }
 
@@ -501,29 +485,15 @@ IMPLEMENT_SIMPLE_AUTOMATION_TEST(
 
 bool FAIAssistantToolsetExecuteToolInvalidToolname::RunTest(const FString& UnusedParameters)
 {
-	ToolsetContainer<UFakeToolset> Container;
-	Container.Input = FakeToolsetInput;
-
-	UAIAssistantToolsetRegistry::ExecuteTool(
-		FakeToolsetName, FString("SimpleToolInvalid"), Container.Input, Container.Result, Container.Error);
-	TestFalse("Tool execution with invalid tool name returns error", Container.Error.IsEmpty());
-	return true;
-}
-
-// Test that invoking a tool with an existing but invalid tool name errors out properly.
-IMPLEMENT_SIMPLE_AUTOMATION_TEST(
-	FAIAssistantToolsetExecuteToolExistingButInvalidToolname,
-	"AI.Assistant.Toolset.ExecuteToolExistingButInvalidToolname",
-	AIAssistantTest::Flags)
-
-bool FAIAssistantToolsetExecuteToolExistingButInvalidToolname::RunTest(const FString& UnusedParameters)
-{
-	ToolsetContainer<UFakeToolset> Container;
-	Container.Input = FakeToolsetInput;
+	ToolsetRegistryContainer Container(UFakeToolset::StaticClass());
+	TSharedPtr<UE::AIAssistant::FToolsetLibraryWrapper> Toolset = Container.ToolsetHandler();
+	TestTrue("Valid toolset handler returned", Toolset.IsValid());
 
+	FString Result;
+	FString Error;
 	UAIAssistantToolsetRegistry::ExecuteTool(
-		FakeToolsetName, FString("Register"), Container.Input, Container.Result, Container.Error);
-	TestFalse("Tool execution with invalid tool name returns error", Container.Error.IsEmpty());
+		FakeToolsetName, FString("SimpleToolInvalid"), FakeToolsetInput, Result, Error);
+	TestFalse("Tool execution with invalid tool name returns error", Error.IsEmpty());
 	return true;
 }
 
@@ -535,74 +505,12 @@ IMPLEMENT_SIMPLE_AUTOMATION_TEST(
 
 bool FAIAssistantToolsetExecuteToolUnregisteredToolset::RunTest(const FString& UnusedParameters)
 {
-	ToolsetContainer<UFakeToolset> Container;
-	Container.Input = FakeToolsetInput;
-
-	TObjectPtr<UFakeToolset> Instance = *Container;
-	Instance->Unregister();
-
+	FString Result;
+	FString Error;
 	AddExpectedError(TEXT("not found"), EAutomationExpectedErrorFlags::Contains);
 	UAIAssistantToolsetRegistry::ExecuteTool(
-		FakeToolsetName, FakeToolName, Container.Input, Container.Result, Container.Error);
-	TestFalse("Tool execution with unregistered toolset returns error", Container.Error.IsEmpty());
-	return true;
-}
-
-// Test that querying toolsets thru registry handler work as expected.
-IMPLEMENT_SIMPLE_AUTOMATION_TEST(
-	FAIAssistantToolsetQueryFromRegistryHandler,
-	"AI.Assistant.Toolset.QueryFromRegistryHandler",
-	AIAssistantTest::Flags)
-
-bool FAIAssistantToolsetQueryFromRegistryHandler::RunTest(const FString& UnusedParameters)
-{
-	ToolsetContainer<UFakeToolset> Container;
-	TObjectPtr<UFakeToolset> Instance = *Container;
-
-	auto AssistantSubsystem = UAIAssistantSubsystem::Get(TEXT("Registry handler test"));
-	TestTrue("Can access AI Assistant subsystem", AssistantSubsystem.HasValue());
-
-	TSharedPtr<UE::AIAssistant::IToolsetHandler>* RegisteredHandlerPtr =
-		AssistantSubsystem.GetValue()->ToolsetRegistry.Find(FakeToolsetName);
-	TSharedPtr<UE::AIAssistant::IToolsetHandler> RegisteredHandler =
-			(RegisteredHandlerPtr) ? *RegisteredHandlerPtr : nullptr;
-	TestTrue("Registered toolset handler found", RegisteredHandler.IsValid());
-
-	TestEqual("Toolset registry handler returns name properly",
-			  Instance->GetName(), RegisteredHandler->GetToolsetName());
-	TestEqual("Toolset registry handler returns version properly",
-			  Instance->GetVersion(), RegisteredHandler->GetToolsetVersion());
-	TestEqual("Toolset registry handler returns JSON schema properly",
-			  Instance->GetJsonSchema(), RegisteredHandler->GetToolsetJsonSchema());
-	return true;
-}
-
-// Test that using registry handler to query toolsets that update version in postinit work as expected.
-IMPLEMENT_SIMPLE_AUTOMATION_TEST(
-	FAIAssistantToolsetQueryFromRegistryHandlerUpdatedVersion,
-	"AI.Assistant.Toolset.QueryFromRegistryHandlerUpdatedVersion",
-	AIAssistantTest::Flags)
-
-bool FAIAssistantToolsetQueryFromRegistryHandlerUpdatedVersion::RunTest(const FString& UnusedParameters)
-{
-	ToolsetContainer<UFakeToolsetUpdatesVersion> Container;
-	TObjectPtr<UFakeToolsetUpdatesVersion> Instance = *Container;
-
-	auto AssistantSubsystem = UAIAssistantSubsystem::Get(TEXT("Registry handler test"));
-	TestTrue("Can access AI Assistant subsystem", AssistantSubsystem.HasValue());
-
-	TSharedPtr<UE::AIAssistant::IToolsetHandler>* RegisteredHandlerPtr =
-		AssistantSubsystem.GetValue()->ToolsetRegistry.Find(FakeToolsetUpdatesVersionName);
-	TSharedPtr<UE::AIAssistant::IToolsetHandler> RegisteredHandler =
-			(RegisteredHandlerPtr) ? *RegisteredHandlerPtr : nullptr;
-	TestTrue("Registered toolset handler found", RegisteredHandler.IsValid());
-
-	TestEqual("Toolset registry handler returns name properly",
-			  Instance->GetName(), RegisteredHandler->GetToolsetName());
-	TestEqual("Toolset registry handler returns version properly",
-			  Instance->GetVersion(), RegisteredHandler->GetToolsetVersion());
-	TestEqual("Toolset registry handler returns JSON schema properly",
-			  Instance->GetJsonSchema(), RegisteredHandler->GetToolsetJsonSchema());
+		FakeToolsetName, FakeToolName, FakeToolsetInput, Result, Error);
+	TestFalse("Tool execution with unregistered toolset returns error", Error.IsEmpty());
 	return true;
 }
 
@@ -614,20 +522,23 @@ IMPLEMENT_SIMPLE_AUTOMATION_TEST(
 
 bool FAIAssistantToolsetExecuteToolWithDefaultParam::RunTest(const FString& UnusedParameters)
 {
-	ToolsetContainer<UFakeToolsetWithDefaultParams> Container;
-	Container.Input = FakeToolsetInput;
+	ToolsetRegistryContainer Container(UFakeToolsetWithDefaultParams::StaticClass());
+	TSharedPtr<UE::AIAssistant::FToolsetLibraryWrapper> Toolset = Container.ToolsetHandler();
+	TestTrue("Valid toolset handler returned", Toolset.IsValid());
 
+	FString Result;
+	FString Error;
 	UAIAssistantToolsetRegistry::ExecuteTool(
-		FString("FakeToolsetWithDefaultParams"), FString("ToolWithDefaults"), Container.Input, Container.Result, Container.Error);
-	TestTrue("Tool execution with default param returns without error", Container.Error.IsEmpty());
+		FString("FakeToolsetWithDefaultParams"), FString("ToolWithDefaults"), FakeToolsetInput, Result, Error);
+	TestTrue("Tool execution with default param returns without error", Error.IsEmpty());
 
 	FString ExpectedOutput = FString(TEXT(R"({"returnValue":"hello: 42"})"));
-	TestEqual("Tool execution returns expected output", Container.Result, ExpectedOutput);
+	TestEqual("Tool execution returns expected output", Result, ExpectedOutput);
 
 	return true;
 }
 
-// Test that invoking a tool with params that don't match input schema errors properly.
+// Test that invoking a tool with params that don't match input schema errors out properly.
 IMPLEMENT_SIMPLE_AUTOMATION_TEST(
 	FAIAssistantToolsetExecuteToolWithBadInput,
 	"AI.Assistant.Toolset.ExecuteToolWithBadInput",
@@ -635,12 +546,16 @@ IMPLEMENT_SIMPLE_AUTOMATION_TEST(
 
 bool FAIAssistantToolsetExecuteToolWithBadInput::RunTest(const FString& UnusedParameters)
 {
-	ToolsetContainer<UFakeToolset> Container;
-	Container.Input = TEXT(R"({"not_expected":"bad_value"})");
+	ToolsetRegistryContainer Container(UFakeToolset::StaticClass());
+	TSharedPtr<UE::AIAssistant::FToolsetLibraryWrapper> Toolset = Container.ToolsetHandler();
+	TestTrue("Valid toolset handler returned", Toolset.IsValid());
 
+	FString Input = TEXT(R"({"not_expected":"bad_value"})");
+	FString Result;
+	FString Error;
 	UAIAssistantToolsetRegistry::ExecuteTool(
-		FakeToolsetName, FakeToolName, Container.Input, Container.Result, Container.Error);
-	TestFalse("Tool execution with invalid param returns with error", Container.Error.IsEmpty());
+		FakeToolsetName, FakeToolName, Input, Result, Error);
+	TestFalse("Tool execution with invalid param returns with error", Error.IsEmpty());
 	return true;
 }
 
diff --git a/Engine/Plugins/Experimental/AIAssistant/Source/AIAssistant/Private/Tests/AIAssistantToolsetTest.h b/Engine/Plugins/Experimental/AIAssistant/Source/AIAssistant/Private/Tests/AIAssistantToolsetTest.h
--- a/Engine/Plugins/Experimental/AIAssistant/Source/AIAssistant/Private/Tests/AIAssistantToolsetTest.h
+++ b/Engine/Plugins/Experimental/AIAssistant/Source/AIAssistant/Private/Tests/AIAssistantToolsetTest.h
@@ -4,13 +4,14 @@
 
 #include "CoreMinimal.h"
 
-#include "AIAssistantToolset.h"
+#include "Dom/JsonObject.h"
+#include "Kismet/BlueprintFunctionLibrary.h"
 
 #include "AIAssistantToolsetTest.generated.h"
 
 /** Fake Toolset that does nothing. */
 UCLASS(Blueprintable)
-class UFakeToolset : public UAIAssistantToolset
+class UFakeToolset : public UBlueprintFunctionLibrary
 {
 	GENERATED_BODY()
 
@@ -18,48 +19,21 @@ public:
 
 	/** Simple tool that does nothing. */
 	UFUNCTION(BlueprintCallable, Category=Fake)
-	FString SimpleTool(const FString& InString)
+	static FString SimpleTool(const FString& InString)
 	{
 		return InString;
 	}
 
 	/** UFunction for internal use, should not be marked as tool. */
 	UFUNCTION(meta=(BlueprintInternalUseOnly=true))
-	void SimpleInternalMethod() { }
-
-	TSharedPtr<FJsonObject> GetExpectedSchema();
-
+	static void SimpleInternalMethod() { }
 };
 
-/** Fake Toolset that does nothing except update version. */
-UCLASS(Blueprintable)
-class UFakeToolsetUpdatesVersion : public UAIAssistantToolset
-{
-	GENERATED_BODY()
-
-public:
-
-	static inline FString MY_VERSION = FString("1.999.the_best");
-
-	/** Simple tool that does nothing. */
-	UFUNCTION(BlueprintCallable, Category=Fake)
-	FString SimpleTool(const FString& InString)
-	{
-		return InString;
-	}
-
-	// Override the ToolsetVersion property.
-	virtual void PostInitProperties() override
-	{
-		ToolsetVersion = MY_VERSION;
-		Super::PostInitProperties();
-	}
-
-};
+static TSharedPtr<FJsonObject> GetFakeToolsetExpectedSchema();
 
 /** Fake Toolset that has method with default parameters. */
 UCLASS(Blueprintable)
-class UFakeToolsetWithDefaultParams : public UAIAssistantToolset
+class UFakeToolsetWithDefaultParams : public UBlueprintFunctionLibrary
 {
 	GENERATED_BODY()
 
@@ -67,28 +41,27 @@ public:
 
 	/** Method with some params with default values. */
 	UFUNCTION(BlueprintCallable, Category=Fake)
-	FString ToolWithDefaults(const FString& InString, int InValue=42)
+	static FString ToolWithDefaults(const FString& InString, int InValue=42)
 	{
 		return FString::Printf(TEXT("%s: %d"), *InString, InValue);
 	}
-
-	TSharedPtr<FJsonObject> GetExpectedSchema();
 };
 
+static TSharedPtr<FJsonObject> GetFakeToolsetWithDefaultParamsExpectedSchema();
+
+
 /** Fake Toolset with no tools. */
 UCLASS(Blueprintable)
-class UFakeToolsetWithNoTools : public UAIAssistantToolset
+class UFakeToolsetWithNoTools : public UBlueprintFunctionLibrary
 {
 	GENERATED_BODY()
-
-public:
-	TSharedPtr<FJsonObject> GetExpectedSchema();
-
 };
 
+static TSharedPtr<FJsonObject> GetFakeToolsetWithNoToolsExpectedSchema();
+
 /** Fake Toolset that has method with TArray parameters. */
 UCLASS(Blueprintable)
-class UFakeToolsetWithArrayParams : public UAIAssistantToolset
+class UFakeToolsetWithArrayParams : public UBlueprintFunctionLibrary
 {
 	GENERATED_BODY()
 
@@ -96,17 +69,17 @@ public:
 
 	/** Method with some params with TArray values. */
 	UFUNCTION(BlueprintCallable, Category=Fake)
-	TArray<float> ToolWithArrays(const TArray<int> ArrayParamOne, TArray<FString> ArrayParamTwo)
+	static TArray<float> ToolWithArrays(const TArray<int> ArrayParamOne, TArray<FString> ArrayParamTwo)
 	{
 		return TArray<float>();
 	}
-
-	TSharedPtr<FJsonObject> GetExpectedSchema();
 };
 
+static TSharedPtr<FJsonObject> GetFakeToolsetWithArrayParamsExpectedSchema();
+
 /** Fake Toolset that has method with TMap parameters. */
 UCLASS(Blueprintable)
-class UFakeToolsetWithMapParams : public UAIAssistantToolset
+class UFakeToolsetWithMapParams : public UBlueprintFunctionLibrary
 {
 	GENERATED_BODY()
 
@@ -114,17 +87,17 @@ public:
 
 	/** Method with some params with TMap values. */
 	UFUNCTION(BlueprintCallable, Category=Fake)
-	TMap<FString, float> ToolWithMaps(const TMap<FString, int> MapParamOne, TMap<FString, FString> MapParamTwo)
+	static TMap<FString, float> ToolWithMaps(const TMap<FString, int> MapParamOne, TMap<FString, FString> MapParamTwo)
 	{
 		return TMap<FString, float>();
 	}
-
-	TSharedPtr<FJsonObject> GetExpectedSchema();
 };
 
+static TSharedPtr<FJsonObject> GetFakeToolsetWithMapParamsExpectedSchema();
+
 /** Fake Toolset that has method with TSet parameters. */
 UCLASS(Blueprintable)
-class UFakeToolsetWithSetParams : public UAIAssistantToolset
+class UFakeToolsetWithSetParams : public UBlueprintFunctionLibrary
 {
 	GENERATED_BODY()
 
@@ -132,10 +105,10 @@ public:
 
 	/** Method with some params with TSet values. */
 	UFUNCTION(BlueprintCallable, Category=Fake)
-	TSet<float> ToolWithSets(const TSet<int> SetParamOne, TSet<FString> SetParamTwo)
+	static TSet<float> ToolWithSets(const TSet<int> SetParamOne, TSet<FString> SetParamTwo)
 	{
 		return TSet<float>();
 	}
-
-	TSharedPtr<FJsonObject> GetExpectedSchema();
 };
+
+static TSharedPtr<FJsonObject> GetFakeToolsetWithSetParamsExpectedSchema();
