[AI Assistant] Cache toolset name in FToolsetLibraryWrapper,
to allow cleanly unregistering wrappers when the backing toolset class is no longer available.

This can happen for instance when the toolset class's Python module is reloaded.

#rnx

#Synced-CL 49891327

diff --git a/Engine/Plugins/Experimental/AIAssistant/Source/AIAssistant/Private/AIAssistantToolset.h b/Engine/Plugins/Experimental/AIAssistant/Source/AIAssistant/Private/AIAssistantToolset.h
--- a/Engine/Plugins/Experimental/AIAssistant/Source/AIAssistant/Private/AIAssistantToolset.h
+++ b/Engine/Plugins/Experimental/AIAssistant/Source/AIAssistant/Private/AIAssistantToolset.h
@@ -22,6 +22,7 @@ namespace UE::AIAssistant
 		FToolsetLibraryWrapper(UClass* LibraryClass) : ToolsetClass(LibraryClass)
 		{
 			GenerateToolCallObjects();
+			ToolsetName = ToolsetClass.Pin() ? ToolsetClass->GetName() : FString();
 		}
 
 		// Executes a tool with the given name and JSON input, returning a future with the JSON
@@ -35,7 +36,7 @@ namespace UE::AIAssistant
 		// Returns the name of the toolset
 		virtual FString GetToolsetName() const override
 		{
-			return ToolsetClass.Pin() ? ToolsetClass->GetName() : FString();
+			return ToolsetName;
 		}
 
 		// Returns the version of the toolset
@@ -71,6 +72,10 @@ namespace UE::AIAssistant
 		// The wrapped toolset function library
 		TWeakObjectPtr<UClass> ToolsetClass;
 
+		// Cached tooolset name to allow unregistering even when the wrapped toolset class is no longer
+		// available.
+		FString ToolsetName;
+
 		// Map of valid tool names to tool call objects.
 		TMap<FString, TSharedPtr<FObjectFunctionToolCall>> Tools;
 
