[AI Assistant] Add stricter constraints on toolset methods to allow as valid tool functions

#rnx

#Synced-CL 49891327

diff --git a/Engine/Plugins/Experimental/AIAssistant/Source/AIAssistant/Private/AIAssistantToolset.cpp b/Engine/Plugins/Experimental/AIAssistant/Source/AIAssistant/Private/AIAssistantToolset.cpp
--- a/Engine/Plugins/Experimental/AIAssistant/Source/AIAssistant/Private/AIAssistantToolset.cpp
+++ b/Engine/Plugins/Experimental/AIAssistant/Source/AIAssistant/Private/AIAssistantToolset.cpp
@@ -186,13 +186,24 @@ namespace UE::AIAssistant
 			TObjectPtr<UFunction> Function = *FuncIt;
 			if (Function)
 			{
+				bool bValidFunction = true;
+
 				// Skip functions marked as for internal use only.
 				if (Function->HasMetaData(TEXT("BlueprintInternalUseOnly")))
 				{
-					continue;
+					bValidFunction = false;
+				}
+
+				// Skip functions unless they are public and static.
+				if (!Function->HasAllFunctionFlags(FUNC_Public | FUNC_Static))
+				{
+					bValidFunction = false;
 				}
 
-				ValidTools.Add(Function);
+				if (bValidFunction)
+				{
+					ValidTools.Add(Function);
+				}
 			}
 		}
 		return ValidTools;
diff --git a/Engine/Plugins/Experimental/AIAssistant/Source/AIAssistant/Private/Tests/AIAssistantToolsetTest.cpp b/Engine/Plugins/Experimental/AIAssistant/Source/AIAssistant/Private/Tests/AIAssistantToolsetTest.cpp
--- a/Engine/Plugins/Experimental/AIAssistant/Source/AIAssistant/Private/Tests/AIAssistantToolsetTest.cpp
+++ b/Engine/Plugins/Experimental/AIAssistant/Source/AIAssistant/Private/Tests/AIAssistantToolsetTest.cpp
@@ -559,6 +559,32 @@ bool FAIAssistantToolsetExecuteToolWithBadInput::RunTest(const FString& UnusedPa
 	return true;
 }
 
+// Test that a toolset doesn't include private, internal and non-static methods.
+IMPLEMENT_SIMPLE_AUTOMATION_TEST(
+	FAIAssistantToolsetValidateToolMethods,
+	"AI.Assistant.Toolset.ValidateToolMethods",
+	AIAssistantTest::Flags)
+
+bool FAIAssistantToolsetValidateToolMethods::RunTest(const FString& UnusedParameters)
+{
+	ToolsetRegistryContainer Container(UFakeToolsetWithPrivateMethods::StaticClass());
+	TSharedPtr<UE::AIAssistant::FToolsetLibraryWrapper> Toolset = Container.ToolsetHandler();
+	if (!TestTrue("Valid toolset handler returned", Toolset.IsValid()))
+	{
+		return false;
+	}
+
+	TArray<FString> ToolNames = Toolset->GetToolNames();
+
+	if (!TestEqual("Only one valid tool should be included", ToolNames.Num(), 1))
+	{
+		return false;
+	}
+	TestEqual("Only \"SimpleTool\" should be valid tool", ToolNames[0], FString("SimpleTool"));
+
+	return true;
+}
+
 // Run python tests
 IMPLEMENT_SIMPLE_AUTOMATION_TEST(
 	FAIAssistantToolsetPythonTests,
diff --git a/Engine/Plugins/Experimental/AIAssistant/Source/AIAssistant/Private/Tests/AIAssistantToolsetTest.h b/Engine/Plugins/Experimental/AIAssistant/Source/AIAssistant/Private/Tests/AIAssistantToolsetTest.h
--- a/Engine/Plugins/Experimental/AIAssistant/Source/AIAssistant/Private/Tests/AIAssistantToolsetTest.h
+++ b/Engine/Plugins/Experimental/AIAssistant/Source/AIAssistant/Private/Tests/AIAssistantToolsetTest.h
@@ -112,3 +112,33 @@ public:
 };
 
 static TSharedPtr<FJsonObject> GetFakeToolsetWithSetParamsExpectedSchema();
+
+/** Fake Toolset with private and non-static methods. */
+UCLASS(Blueprintable)
+class UFakeToolsetWithPrivateMethods : public UBlueprintFunctionLibrary
+{
+	GENERATED_BODY()
+
+public:
+
+	/** Simple tool that does nothing. */
+	UFUNCTION(BlueprintCallable, Category=Fake)
+	static FString SimpleTool(const FString& InString)
+	{
+		return InString;
+	}
+
+	/** UFunction for internal use, should not be marked as tool. */
+	UFUNCTION(meta=(BlueprintInternalUseOnly=true))
+	static void SimpleInternalMethod() {}
+
+	/** UFunction non-static method, should not be marked as tool. */
+	UFUNCTION()
+	void SimpleNonStaticMethod() {}
+
+
+private:
+	/** UFunction private method, should not be marked as tool. */
+	UFUNCTION()
+	static void SimplePrivateMethod() {}
+};
