[AI Assistant] Don't raise error if a tool call doesn't have a returned value

#rnx

#Synced-CL 49899510

diff --git a/Engine/Plugins/Experimental/AIAssistant/Source/AIAssistant/Private/AIAssistantToolset.cpp b/Engine/Plugins/Experimental/AIAssistant/Source/AIAssistant/Private/AIAssistantToolset.cpp
--- a/Engine/Plugins/Experimental/AIAssistant/Source/AIAssistant/Private/AIAssistantToolset.cpp
+++ b/Engine/Plugins/Experimental/AIAssistant/Source/AIAssistant/Private/AIAssistantToolset.cpp
@@ -132,17 +132,21 @@ namespace UE::AIAssistant
 			return;
 		}
 
-		if (!ToolCallResult.HasValue())
+		if (ToolCallResult.HasValue())
 		{
-			OutError = FString::Printf(TEXT("No return value from executing tool %s"), *ToolName);
+			ToolCallReturnValue = ToolCallResult.StealValue();
+		}
+		else
+		{
+			OutError = FString::Printf(TEXT("Invalid return value from executing tool %s"), *ToolName);
 			return;
 		}
 
-		ToolCallReturnValue = ToolCallResult.StealValue();
 		if (!ToolCallReturnValue.IsValid())
 		{
-			OutError = FString::Printf(TEXT("No return value from executing tool %s"), *ToolName);
-			return;
+			// There was a returned value, but it was null.
+			ToolCallReturnValue = MakeShared<FJsonObject>();
+			ToolCallReturnValue->SetStringField(FString("returnValue"), FString("null"));
 		}
 
 		if (!FJsonSerializer::Serialize(
diff --git a/Engine/Plugins/Experimental/AIAssistant/Source/AIAssistant/Private/Tests/AIAssistantToolsetTest.cpp b/Engine/Plugins/Experimental/AIAssistant/Source/AIAssistant/Private/Tests/AIAssistantToolsetTest.cpp
--- a/Engine/Plugins/Experimental/AIAssistant/Source/AIAssistant/Private/Tests/AIAssistantToolsetTest.cpp
+++ b/Engine/Plugins/Experimental/AIAssistant/Source/AIAssistant/Private/Tests/AIAssistantToolsetTest.cpp
@@ -67,25 +67,46 @@ TSharedPtr<FJsonObject> GetFakeToolsetExpectedSchema()
 	FJsonDomBuilder::FObject SchemaBuilder = MakeSchemaBuilder(
 		FString(TEXT("FakeToolset")), FString("1.0"), FString(TEXT("Fake Toolset that does nothing.")));
 
-	FJsonDomBuilder::FObject InputSchema = MakeSchema();
-	FJsonDomBuilder::FObject Properties;
 
-	FString ParamName = TEXT("inString");
+	FJsonDomBuilder::FArray ToolsArray;
+	// Add SimpleTool
+	{
+		FJsonDomBuilder::FObject InputSchema = MakeSchema();
+		FJsonDomBuilder::FObject Properties;
 
-	Properties.Set(ParamName, MakeProperty(TEXT("string")));
-	AddPropertiesToSchema(InputSchema, Properties);
-	AddRequiredListToSchema(InputSchema, TArray<FString>({ParamName}));
+		FString ParamName = TEXT("inString");
+		Properties.Set(ParamName, MakeProperty(TEXT("string")));
+		AddPropertiesToSchema(InputSchema, Properties);
+		AddRequiredListToSchema(InputSchema, TArray<FString>({ParamName}));
 
-	FJsonDomBuilder::FObject OutputSchema = MakeSchema();
-	AddPropertiesToSchema(OutputSchema, MakeNamedProperty(ReturnValueName, TEXT("string")));
-	AddRequiredListToSchema(OutputSchema, TArray<FString>({ReturnValueName}));
+		FJsonDomBuilder::FObject OutputSchema = MakeSchema();
+		AddPropertiesToSchema(OutputSchema, MakeNamedProperty(ReturnValueName, TEXT("string")));
+		AddRequiredListToSchema(OutputSchema, TArray<FString>({ReturnValueName}));
 
-	FJsonDomBuilder::FObject ToolEntry = MakeToolSchemaEntry(
-		TEXT("FakeToolset.SimpleTool"), TEXT("Simple tool that does nothing."),
-		InputSchema, OutputSchema);
+		FJsonDomBuilder::FObject ToolEntry = MakeToolSchemaEntry(
+			TEXT("FakeToolset.SimpleTool"), TEXT("Simple tool that does nothing."),
+			InputSchema, OutputSchema);
 
-	FJsonDomBuilder::FArray ToolsArray;
-	AddToolSchema(SchemaBuilder, ToolsArray, ToolEntry);
+		AddToolSchema(SchemaBuilder, ToolsArray, ToolEntry);
+	}
+
+	// Add SimpleNoReturnTool
+	{
+		FJsonDomBuilder::FObject InputSchema = MakeSchema();
+		FJsonDomBuilder::FObject Properties;
+
+		FString ParamName = TEXT("inString");
+		Properties.Set(ParamName, MakeProperty(TEXT("string")));
+		AddPropertiesToSchema(InputSchema, Properties);
+		AddRequiredListToSchema(InputSchema, TArray<FString>({ParamName}));
+
+		FJsonDomBuilder::FObject ToolEntry;
+		ToolEntry.Set(TEXT("description"), TEXT("Simple tool that does nothing and returns nothing."));
+		ToolEntry.Set(TEXT("name"), TEXT("FakeToolset.SimpleNoReturnTool"));
+		ToolEntry.Set(TEXT("inputSchema"), InputSchema);
+
+		AddToolSchema(SchemaBuilder, ToolsArray, ToolEntry);
+	}
 
 	return SchemaBuilder.AsJsonObject();
 }
@@ -514,6 +535,30 @@ bool FAIAssistantToolsetExecuteToolUnregisteredToolset::RunTest(const FString& U
 	return true;
 }
 
+// Test that invoking a tool that doesn't return anything doesn't raise errors.
+IMPLEMENT_SIMPLE_AUTOMATION_TEST(
+	FAIAssistantToolsetExecuteToolWithNoReturn,
+	"AI.Assistant.Toolset.ExecuteToolWithNoReturn",
+	AIAssistantTest::Flags)
+
+bool FAIAssistantToolsetExecuteToolWithNoReturn::RunTest(const FString& UnusedParameters)
+{
+	ToolsetRegistryContainer Container(UFakeToolset::StaticClass());
+	TSharedPtr<UE::AIAssistant::FToolsetLibraryWrapper> Toolset = Container.ToolsetHandler();
+	TestTrue("Valid toolset handler returned", Toolset.IsValid());
+
+	FString Result;
+	FString Error;
+	UAIAssistantToolsetRegistry::ExecuteTool(
+		FakeToolsetName, FString("SimpleNoReturnTool"), FakeToolsetInput, Result, Error);
+	TestTrue("Tool execution with no return value runs without error", Error.IsEmpty());
+
+	FString ExpectedOutput = FString(TEXT(R"({"returnValue":"null"})"));
+	TestEqual("Tool execution returns expected output", Result, ExpectedOutput);
+
+	return true;
+}
+
 // Test that invoking a tool with default values for params provides the proper default value.
 IMPLEMENT_SIMPLE_AUTOMATION_TEST(
 	FAIAssistantToolsetExecuteToolWithDefaultParam,
diff --git a/Engine/Plugins/Experimental/AIAssistant/Source/AIAssistant/Private/Tests/AIAssistantToolsetTest.h b/Engine/Plugins/Experimental/AIAssistant/Source/AIAssistant/Private/Tests/AIAssistantToolsetTest.h
--- a/Engine/Plugins/Experimental/AIAssistant/Source/AIAssistant/Private/Tests/AIAssistantToolsetTest.h
+++ b/Engine/Plugins/Experimental/AIAssistant/Source/AIAssistant/Private/Tests/AIAssistantToolsetTest.h
@@ -24,6 +24,13 @@ public:
 		return InString;
 	}
 
+	/** Simple tool that does nothing and returns nothing. */
+	UFUNCTION(BlueprintCallable, Category=Fake)
+	static void SimpleNoReturnTool(const FString& InString)
+	{
+		return;
+	}
+
 	/** UFunction for internal use, should not be marked as tool. */
 	UFUNCTION(meta=(BlueprintInternalUseOnly=true))
 	static void SimpleInternalMethod() { }
