[AI Assistant] Support for tools that raise errors when run

- Add wrapper to tools written in Python to catch exceptions to be returned to assistant as error.
- Add standard "error" field to tool output schemas: runtime errors raised by tools will go here, empty if no error occurred.
- Add tests for tools that raise errors

#rnx

#Synced-CL 50097725

diff --git a/Engine/Plugins/Experimental/AIAssistant/Source/AIAssistant/Private/AIAssistantToolset.cpp b/Engine/Plugins/Experimental/AIAssistant/Source/AIAssistant/Private/AIAssistantToolset.cpp
--- a/Engine/Plugins/Experimental/AIAssistant/Source/AIAssistant/Private/AIAssistantToolset.cpp
+++ b/Engine/Plugins/Experimental/AIAssistant/Source/AIAssistant/Private/AIAssistantToolset.cpp
@@ -10,13 +10,50 @@
 #include "Policies/CondensedJsonPrintPolicy.h"
 #include "Serialization/JsonWriter.h"
 #include "Serialization/JsonSerializer.h"
+#include "Blueprint/BlueprintExceptionInfo.h" 
+#include "AIAssistantDelegateHandle.h"
 
 #include "AIAssistantLog.h"
-#include "AIAssistantSubsystem.h"
 #include "AIAssistantToolCall.h"
 
 namespace UE::AIAssistant
 {
+	namespace
+	{
+		// Register for tracking of script exceptions for the lifetime of this object.
+		class FToolCallScriptExceptionHandler
+		{
+		  public:
+			FToolCallScriptExceptionHandler()
+				: ScriptExceptionHandle(
+					FDelegateHandleRaii::Create(
+						FBlueprintCoreDelegates::OnScriptException,
+						FBlueprintCoreDelegates::OnScriptException.AddRaw(
+							this, &FToolCallScriptExceptionHandler::OnScriptException)))
+			{}
+			~FToolCallScriptExceptionHandler() = default;
+
+			// Track script exceptions and copy into local storage.
+			void OnScriptException(
+				const UObject* ActiveObject, const FFrame& StackFrame, const FBlueprintExceptionInfo& Info)
+			{
+				if (Info.GetType() == EBlueprintExceptionType::UserRaisedError)
+				{
+					ScriptException = Info.GetDescription().ToString();
+				}
+			}
+
+			FString GetException() const
+			{
+				return ScriptException;
+			}
+
+		  private:
+			FDelegateHandleRaii ScriptExceptionHandle;
+			FString ScriptException;
+		};
+	}
+
 	TFuture<TValueOrError<FString, FString>> FToolsetLibraryWrapper::ExecuteTool(
 		const FString& ToolName, const FString& JsonInput)
 	{
@@ -67,9 +104,39 @@ namespace UE::AIAssistant
 			check(Function);
 
 			PRAGMA_DISABLE_DEPRECATION_WARNINGS
-				TSharedPtr<FJsonObject> ToolEntry = FJsonSchemaGenerator::UStructToJsonSchemaObject(Function);
+			TSharedPtr<FJsonObject> ToolEntry = FJsonSchemaGenerator::UStructToJsonSchemaObject(Function);
 			PRAGMA_ENABLE_DEPRECATION_WARNINGS
 
+			// Insert error field into outputSchema, creating it if it doesn't exist.
+			TSharedPtr<FJsonObject> OutputSchemaProperties;
+			TSharedPtr<FJsonObject> OutputSchema;
+			FString OutputSchemaName = TEXT("outputSchema");
+			if (!ToolEntry->HasField(OutputSchemaName))
+			{
+				OutputSchema = MakeShared<FJsonObject>();
+				ToolEntry->SetObjectField(OutputSchemaName, OutputSchema);
+				OutputSchema->SetStringField(TEXT("type"), TEXT("object"));
+				OutputSchemaProperties = MakeShared<FJsonObject>();
+				TSharedPtr<FJsonObject> ReturnValueProperty = MakeShared<FJsonObject>();
+				ReturnValueProperty->SetField(TEXT("type"), MakeShared<FJsonValueNull>());
+				OutputSchemaProperties->SetObjectField(TEXT("returnValue"), ReturnValueProperty);
+				OutputSchema->SetObjectField(TEXT("properties"), OutputSchemaProperties);
+			}
+			else
+			{
+				OutputSchema = ToolEntry->GetObjectField(OutputSchemaName);
+				check(OutputSchema->HasField(TEXT("properties")));
+				OutputSchemaProperties = OutputSchema->GetObjectField(TEXT("properties"));
+			}
+
+			const FString ErrorFieldName = FString(TEXT("error"));
+			if (!OutputSchemaProperties->HasField(ErrorFieldName))
+			{
+				TSharedPtr<FJsonObject> ErrorProperty = MakeShared<FJsonObject>();
+				ErrorProperty->SetStringField(TEXT("type"), TEXT("string"));
+				OutputSchemaProperties->SetObjectField(ErrorFieldName, ErrorProperty);
+			}
+
 			ToolsArray.Add(MakeShared<FJsonValueObject>(ToolEntry));
 		}
 		Schema->SetArrayField(FString(TEXT("tools")), ToolsArray);
@@ -123,6 +190,11 @@ namespace UE::AIAssistant
 		FObjectFunctionToolCall::FFunctionInputParamsJson JsonParams;
 		JsonParams.Set<FString>(JsonInput);
 
+		// Capture script errors raised by the tool during execution.  Note: script errors raised by tools
+		// running outside of the game thread will not be captured correctly.
+		FToolCallScriptExceptionHandler ExceptionHandler;
+		ensureAlways(IsInGameThread());
+
 		FObjectFunctionToolCall::FToolCallResult ToolCallResult = ToolCall->Execute(JsonParams);
 
 		if (ToolCallResult.HasError())
@@ -149,6 +221,13 @@ namespace UE::AIAssistant
 			ToolCallReturnValue->SetStringField(FString("returnValue"), FString("null"));
 		}
 
+		// If the tool raised an error, return that.
+		if (FString ScriptException = ExceptionHandler.GetException();
+			!ScriptException.IsEmpty())
+		{
+			ToolCallReturnValue->SetStringField(TEXT("error"), MoveTemp(ScriptException));
+		}
+
 		if (!FJsonSerializer::Serialize(
 				ToolCallReturnValue.ToSharedRef(),
 				TJsonWriterFactory<TCHAR, TCondensedJsonPrintPolicy<TCHAR>>::Create(&OutResult)))
diff --git a/Engine/Plugins/Experimental/AIAssistant/Source/AIAssistant/Private/AIAssistantToolset.h b/Engine/Plugins/Experimental/AIAssistant/Source/AIAssistant/Private/AIAssistantToolset.h
--- a/Engine/Plugins/Experimental/AIAssistant/Source/AIAssistant/Private/AIAssistantToolset.h
+++ b/Engine/Plugins/Experimental/AIAssistant/Source/AIAssistant/Private/AIAssistantToolset.h
@@ -72,7 +72,7 @@ namespace UE::AIAssistant
 		// The wrapped toolset function library
 		TWeakObjectPtr<UClass> ToolsetClass;
 
-		// Cached tooolset name to allow unregistering even when the wrapped toolset class is no longer
+		// Cached toolset name to allow unregistering even when the wrapped toolset class is no longer
 		// available.
 		FString ToolsetName;
 
diff --git a/Engine/Plugins/Experimental/AIAssistant/Source/AIAssistant/Private/Tests/AIAssistantToolsetTest.cpp b/Engine/Plugins/Experimental/AIAssistant/Source/AIAssistant/Private/Tests/AIAssistantToolsetTest.cpp
--- a/Engine/Plugins/Experimental/AIAssistant/Source/AIAssistant/Private/Tests/AIAssistantToolsetTest.cpp
+++ b/Engine/Plugins/Experimental/AIAssistant/Source/AIAssistant/Private/Tests/AIAssistantToolsetTest.cpp
@@ -10,6 +10,7 @@
 #include "IPythonScriptPlugin.h"
 #include "PythonScriptTypes.h"
 #include "Policies/CondensedJsonPrintPolicy.h"
+#include "Serialization/JsonReader.h"
 #include "Serialization/JsonWriter.h"
 #include "Serialization/JsonSerializer.h"
 
@@ -27,6 +28,7 @@ const FString FakeToolsetUpdatesVersionName(TEXT("FakeToolsetUpdatesVersion"));
 const FString FakeToolsetInput(TEXT(R"({"inString":"hello"})"));
 const FString FakeToolName(TEXT("SimpleTool"));
 const FString ReturnValueName(TEXT("returnValue"));
+const FString ErrorValueName(TEXT("error"));
 
 class ToolsetRegistryContainer
 {
@@ -80,7 +82,10 @@ TSharedPtr<FJsonObject> GetFakeToolsetExpectedSchema()
 		AddRequiredListToSchema(InputSchema, TArray<FString>({ParamName}));
 
 		FJsonDomBuilder::FObject OutputSchema = MakeSchema();
-		AddPropertiesToSchema(OutputSchema, MakeNamedProperty(ReturnValueName, TEXT("string")));
+		FJsonDomBuilder::FObject OutputProperties;
+		OutputProperties.Set(ReturnValueName, MakeProperty(TEXT("string")));
+		OutputProperties.Set(ErrorValueName, MakeProperty(TEXT("string")));
+		AddPropertiesToSchema(OutputSchema, OutputProperties);
 		AddRequiredListToSchema(OutputSchema, TArray<FString>({ReturnValueName}));
 
 		FJsonDomBuilder::FObject ToolEntry = MakeToolSchemaEntry(
@@ -100,10 +105,17 @@ TSharedPtr<FJsonObject> GetFakeToolsetExpectedSchema()
 		AddPropertiesToSchema(InputSchema, Properties);
 		AddRequiredListToSchema(InputSchema, TArray<FString>({ParamName}));
 
-		FJsonDomBuilder::FObject ToolEntry;
-		ToolEntry.Set(TEXT("description"), TEXT("Simple tool that does nothing and returns nothing."));
-		ToolEntry.Set(TEXT("name"), TEXT("FakeToolset.SimpleNoReturnTool"));
-		ToolEntry.Set(TEXT("inputSchema"), InputSchema);
+		FJsonDomBuilder::FObject OutputSchema = MakeSchema();
+		FJsonDomBuilder::FObject OutputProperties;
+		FJsonDomBuilder::FObject NullPropertyType;
+		NullPropertyType.Set(TEXT("type"), nullptr);
+		OutputProperties.Set(ReturnValueName, NullPropertyType);
+		OutputProperties.Set(ErrorValueName, MakeProperty(TEXT("string")));
+		AddPropertiesToSchema(OutputSchema, OutputProperties);
+
+		FJsonDomBuilder::FObject ToolEntry = MakeToolSchemaEntry(
+			TEXT("FakeToolset.SimpleNoReturnTool"), TEXT("Simple tool that does nothing and returns nothing."),
+			InputSchema, OutputSchema);
 
 		AddToolSchema(SchemaBuilder, ToolsArray, ToolEntry);
 	}
@@ -129,7 +141,10 @@ TSharedPtr<FJsonObject> GetFakeToolsetWithDefaultParamsExpectedSchema()
 	AddRequiredListToSchema(InputSchema, TArray<FString>({ParamName1}));
 
 	FJsonDomBuilder::FObject OutputSchema = MakeSchema();
-	AddPropertiesToSchema(OutputSchema, MakeNamedProperty(ReturnValueName, TEXT("string")));
+	FJsonDomBuilder::FObject OutputProperties;
+	OutputProperties.Set(ReturnValueName, MakeProperty(TEXT("string")));
+	OutputProperties.Set(ErrorValueName, MakeProperty(TEXT("string")));
+	AddPropertiesToSchema(OutputSchema, OutputProperties);
 	AddRequiredListToSchema(OutputSchema, TArray<FString>({ReturnValueName}));
 
 	FJsonDomBuilder::FObject ToolEntry = MakeToolSchemaEntry(
@@ -163,6 +178,7 @@ TSharedPtr<FJsonObject> GetFakeToolsetWithArrayParamsExpectedSchema()
 	FJsonDomBuilder::FObject OutputSchema = MakeSchema();
 	FJsonDomBuilder::FObject OutputProperties;
 	OutputProperties.Set(ReturnValueName, MakeArrayProperty(TEXT("number")));
+	OutputProperties.Set(ErrorValueName, MakeProperty(TEXT("string")));
 	AddPropertiesToSchema(OutputSchema, OutputProperties);
 	AddRequiredListToSchema(OutputSchema, TArray<FString>({ReturnValueName}));
 
@@ -197,6 +213,7 @@ TSharedPtr<FJsonObject> GetFakeToolsetWithSetParamsExpectedSchema()
 	FJsonDomBuilder::FObject OutputSchema = MakeSchema();
 	FJsonDomBuilder::FObject OutputProperties;
 	OutputProperties.Set(ReturnValueName, MakeSetProperty(TEXT("number")));
+	OutputProperties.Set(ErrorValueName, MakeProperty(TEXT("string")));
 	AddPropertiesToSchema(OutputSchema, OutputProperties);
 	AddRequiredListToSchema(OutputSchema, TArray<FString>({ReturnValueName}));
 
@@ -231,6 +248,7 @@ TSharedPtr<FJsonObject> GetFakeToolsetWithMapParamsExpectedSchema()
 	FJsonDomBuilder::FObject OutputSchema = MakeSchema();
 	FJsonDomBuilder::FObject OutputProperties;
 	OutputProperties.Set(ReturnValueName, MakeMapProperty(TEXT("number")));
+	OutputProperties.Set(ErrorValueName, MakeProperty(TEXT("string")));
 	AddPropertiesToSchema(OutputSchema, OutputProperties);
 	AddRequiredListToSchema(OutputSchema, TArray<FString>({ReturnValueName}));
 
@@ -630,6 +648,52 @@ bool FAIAssistantToolsetValidateToolMethods::RunTest(const FString& UnusedParame
 	return true;
 }
 
+// Invoke python tool that produces errors.
+IMPLEMENT_SIMPLE_AUTOMATION_TEST(
+	FAIAssistantToolsetExecutePythonToolWithError,
+	"AI.Assistant.Toolset.ExecutePythonToolWithError",
+	AIAssistantTest::Flags)
+
+bool FAIAssistantToolsetExecutePythonToolWithError::RunTest(const FString& UnusedParameters)
+{
+	if (!IPythonScriptPlugin::Get()->IsPythonAvailable())
+	{
+		return true;
+	}
+
+	FPythonCommandEx PythonCommand = FPythonCommandEx();
+	PythonCommand.ExecutionMode = EPythonCommandExecutionMode::ExecuteFile;
+	PythonCommand.Command = TEXT(
+		"from ai_assistant.tests.demo_toolset import ErrorProneToolset\n"
+		"import unreal; unreal.AIAssistantToolsetRegistry.register_toolset_class(ErrorProneToolset)\n");
+	if (!TestTrue("Register ErrorProneToolset", IPythonScriptPlugin::Get()->ExecPythonCommandEx(PythonCommand)))
+	{
+		return false;
+	}
+
+	FString Input = TEXT(R"({"msg":"Hi"})");
+	FString Result;
+	FString Error;
+	UAIAssistantToolsetRegistry::ExecuteTool("ErrorProneToolset", "bad_tool", Input, Result, Error);
+	TestTrue("Tool was called correctly", Error.IsEmpty());
+
+	TSharedPtr<FJsonObject> ResultObject;
+	if (TestTrue("Deserialized tool output correctly",
+				 FJsonSerializer::Deserialize(TJsonReaderFactory<>::Create(Result), ResultObject) &&
+				 ResultObject.IsValid()))
+	{
+		FString ToolError = ResultObject->GetStringField(TEXT("error"));
+		TestFalse("Tool raised internal error", ToolError.IsEmpty());
+	}
+
+	PythonCommand.Command = TEXT(
+		"from ai_assistant.tests.demo_toolset import ErrorProneToolset\n"
+		"import unreal; unreal.AIAssistantToolsetRegistry.unregister_toolset_class(ErrorProneToolset)\n");
+	TestTrue("Unregister ErrorProneToolset", IPythonScriptPlugin::Get()->ExecPythonCommandEx(PythonCommand));
+
+	return true;
+}
+
 // Run python tests
 IMPLEMENT_SIMPLE_AUTOMATION_TEST(
 	FAIAssistantToolsetPythonTests,
