[AI Assistant] Remove unused Python executor

Also, do not export the private UAIAssistantSubsystem from the library.

#Synced-CL 50282818

diff --git a/Engine/Plugins/Experimental/AIAssistant/Source/AIAssistant/Private/AIAssistantPythonExecutor.cpp b/Engine/Plugins/Experimental/AIAssistant/Source/AIAssistant/Private/AIAssistantPythonExecutor.cpp
deleted file mode 100644
diff --git a/Engine/Plugins/Experimental/AIAssistant/Source/AIAssistant/Private/AIAssistantPythonExecutor.h b/Engine/Plugins/Experimental/AIAssistant/Source/AIAssistant/Private/AIAssistantPythonExecutor.h
deleted file mode 100644
diff --git a/Engine/Plugins/Experimental/AIAssistant/Source/AIAssistant/Private/AIAssistantSubsystem.cpp b/Engine/Plugins/Experimental/AIAssistant/Source/AIAssistant/Private/AIAssistantSubsystem.cpp
--- a/Engine/Plugins/Experimental/AIAssistant/Source/AIAssistant/Private/AIAssistantSubsystem.cpp
+++ b/Engine/Plugins/Experimental/AIAssistant/Source/AIAssistant/Private/AIAssistantSubsystem.cpp
@@ -11,14 +11,12 @@
 
 #include "AIAssistant.h"
 #include "AIAssistantLog.h"
-#include "AIAssistantPythonExecutor.h"
 #include "AIAssistantWebBrowser.h"
 
 #include UE_INLINE_GENERATED_CPP_BY_NAME(AIAssistantSubsystem)
 
 
 using namespace UE::AIAssistant;
-using namespace UE::AIAssistant::PythonExecutor;
 
 //
 // Statics.
@@ -80,20 +78,6 @@ void UAIAssistantSubsystem::Deinitialize()
 	Super::Deinitialize();
 }
 
-/*no:static*/ FString UAIAssistantSubsystem::ExecutePythonScriptViaJavaScript(const FString& Code)
-{
-	FString CodeOutput;
-	ExecutePythonScript(Code, &CodeOutput);
-
-	if (CodeOutput.IsEmpty())
-	{
-		CodeOutput = TEXT("Code executed successfully.");
-	}
-
-	return CodeOutput;
-}
-
-
 /*no:static*/ void UAIAssistantSubsystem::ShowContextMenuViaJavaScript(const FString& SelectedString, const int32 ClientX, const int32 ClientY) const
 {
 	GetAIAssistantModule().ShowContextMenu(SelectedString, FVector2f(ClientX, ClientY));
diff --git a/Engine/Plugins/Experimental/AIAssistant/Source/AIAssistant/Private/AIAssistantSubsystem.h b/Engine/Plugins/Experimental/AIAssistant/Source/AIAssistant/Private/AIAssistantSubsystem.h
--- a/Engine/Plugins/Experimental/AIAssistant/Source/AIAssistant/Private/AIAssistantSubsystem.h
+++ b/Engine/Plugins/Experimental/AIAssistant/Source/AIAssistant/Private/AIAssistantSubsystem.h
@@ -20,7 +20,7 @@ class SAIAssistantWebBrowser;
 //
 
 UCLASS(BlueprintType)
-class AIASSISTANT_API UAIAssistantSubsystem : public UEditorSubsystem
+class UAIAssistantSubsystem : public UEditorSubsystem
 {
 	GENERATED_BODY()
 
@@ -32,10 +32,6 @@ public:
 	//
 	// JavaScript exposed functions.
 	//
-
-	// See NOTE_JAVASCRIPT_CPP_FUNCTIONS in C++ code for how to call this from JavaScript.
-	UFUNCTION(BlueprintCallable, Category="JavaScript")
-	FString ExecutePythonScriptViaJavaScript(const FString& Code);
 	
 	// See NOTE_JAVASCRIPT_CPP_FUNCTIONS in C++ code for how to call this from JavaScript.
 	UFUNCTION(BlueprintCallable, Category="JavaScript")
