Multiple Undo Stacks Support

Add support for multiple independent undo stacks. This allows AI Agents to create their own transaction stacks and temporarily override the global stack to keep their work segmented from the human developer's undo stack.

#Synced-CL 50288418

diff --git a/Engine/Plugins/Experimental/AIAssistant/Source/AIAssistant/Private/AIAssistantTransactionBufferManager.cpp b/Engine/Plugins/Experimental/AIAssistant/Source/AIAssistant/Private/AIAssistantTransactionBufferManager.cpp
new file mode 100644
--- /dev/null
+++ b/Engine/Plugins/Experimental/AIAssistant/Source/AIAssistant/Private/AIAssistantTransactionBufferManager.cpp
@@ -0,0 +1,293 @@
+// Copyright Epic Games, Inc. All Rights Reserved.
+
+#include "AIAssistantTransactionBufferManager.h"
+
+#include "CoreGlobals.h"
+#include "Editor.h"
+#include "Editor/EditorEngine.h"
+#include "Editor/TransBuffer.h"
+#include "HAL/IConsoleManager.h"
+#include "Misc/ConfigCacheIni.h"
+
+#include "AIAssistantLog.h"
+
+namespace
+{
+	// Encapsulates all mutable state for the transaction buffer manager.
+	struct FTransactionBufferManagerState
+	{
+		// The original global transaction buffer.
+		TObjectPtr<UTransactor> SavedGlobalBuffer;
+
+		// Whether an override is currently active.
+		bool bIsOverrideActive = false;
+
+		// Registry of plugin-created transaction buffers.
+		TMap<FName, TObjectPtr<UTransBuffer>> BufferRegistry;
+
+		// Reset all state to initial values.
+		void Reset()
+		{
+			SavedGlobalBuffer = nullptr;
+			bIsOverrideActive = false;
+			BufferRegistry.Empty();
+		}
+	};
+
+	static FTransactionBufferManagerState TransactionBufferManagerState;
+}  // namespace
+
+namespace UE::AIAssistant
+{
+
+static constexpr int32 DefaultUndoBufferSize = 256;
+
+TObjectPtr<UTransBuffer> FTransactionBufferManager::GetTransactionBuffer(const FName& BufferName)
+{
+	check(IsInGameThread());
+
+	if (TObjectPtr<UTransBuffer>* Result =
+		TransactionBufferManagerState.BufferRegistry.Find(BufferName))
+	{
+		return *Result;
+	}
+	return nullptr;
+}
+
+TObjectPtr<UTransBuffer> FTransactionBufferManager::GetOrCreateTransactionBuffer(
+	const FName& BufferName)
+{
+	check(IsInGameThread());
+	check(GConfig);
+
+	if (TObjectPtr<UTransBuffer> ExistingBuffer = GetTransactionBuffer(BufferName))
+	{
+		return ExistingBuffer;
+	}
+
+	TObjectPtr<UTransBuffer> NewBuffer = NewObject<UTransBuffer>();
+	check(NewBuffer);
+
+	// This replicates the behavior of the global undo buffer initialization found in
+	// UEditorEngine::CreateTrans().
+	int32 UndoBufferSize;
+	if (!GConfig->GetInt(TEXT("Undo"), TEXT("UndoBufferSize"),
+		UndoBufferSize, GEditorPerProjectIni))
+	{
+		UndoBufferSize = DefaultUndoBufferSize;
+	}
+	NewBuffer->AddToRoot();
+	NewBuffer->Initialize(UndoBufferSize);
+	TransactionBufferManagerState.BufferRegistry.Add(BufferName, NewBuffer);
+	return NewBuffer;
+}
+
+bool FTransactionBufferManager::DestroyTransactionBuffer(const FName& BufferName)
+{
+	check(IsInGameThread());
+	check(GEditor);
+
+	TObjectPtr<UTransBuffer>* BufferPtr =
+		TransactionBufferManagerState.BufferRegistry.Find(BufferName);
+	if (!BufferPtr)
+	{
+		return false;
+	}
+
+	UTransBuffer* Buffer = BufferPtr->Get();
+
+	// Cannot destroy a buffer that is currently the active override
+	if (GEditor->Trans == Buffer)
+	{
+		UE_LOG(
+			LogAIAssistant, Warning,
+			TEXT("Cannot destroy transaction buffer '%s' while it is active as override"),
+			*BufferName.ToString());
+		return false;
+	}
+
+	// Cannot destroy the saved global buffer while an override is active
+	if (TransactionBufferManagerState.SavedGlobalBuffer == Buffer)
+	{
+		UE_LOG(
+			LogAIAssistant, Warning,
+			TEXT("Cannot destroy transaction buffer '%s' while it is saved as the global buffer"),
+			*BufferName.ToString());
+		return false;
+	}
+
+	// Remove from root to allow garbage collection
+	if (Buffer)
+	{
+		Buffer->RemoveFromRoot();
+	}
+
+	TransactionBufferManagerState.BufferRegistry.Remove(BufferName);
+	return true;
+}
+
+void FTransactionBufferManager::Shutdown()
+{
+	check(IsInGameThread());
+
+	// Restore override back to original global buffer if active
+	if (TransactionBufferManagerState.bIsOverrideActive)
+	{
+		RestoreGlobalBuffer();
+	}
+
+	// Remove all buffers from root and clear registry
+	for (auto& Pair : TransactionBufferManagerState.BufferRegistry)
+	{
+		if (Pair.Value)
+		{
+			Pair.Value->RemoveFromRoot();
+		}
+	}
+	TransactionBufferManagerState.Reset();
+}
+
+bool FTransactionBufferManager::SetOverrideBuffer(UTransactor* InOverrideBuffer)
+{
+	check(IsInGameThread());
+	check(GEditor);
+
+	if (!InOverrideBuffer)
+	{
+		UE_LOG(
+			LogAIAssistant, Warning,
+			TEXT("SetOverrideBuffer: OverrideBuffer is null, use RestoreGlobalBuffer() to restore"));
+		return false;
+	}
+
+	if (TransactionBufferManagerState.bIsOverrideActive)
+	{
+		UE_LOG(
+			LogAIAssistant, Warning,
+			TEXT("SetOverrideBuffer: Cannot set override while another override is active. "
+				 "Call RestoreGlobalBuffer() first."));
+		return false;
+	}
+
+	TransactionBufferManagerState.SavedGlobalBuffer = GEditor->Trans;
+	TransactionBufferManagerState.bIsOverrideActive = true;
+	GEditor->Trans = InOverrideBuffer;
+	return true;
+}
+
+void FTransactionBufferManager::RestoreGlobalBuffer()
+{
+	check(IsInGameThread());
+	check(GEditor);
+
+	if (!TransactionBufferManagerState.bIsOverrideActive)
+	{
+		UE_LOG(
+			LogAIAssistant, Warning,
+			TEXT("RestoreGlobalBuffer: No override active (global buffer was never saved)"));
+		return;
+	}
+
+	if (TransactionBufferManagerState.SavedGlobalBuffer
+		&& !IsValid(TransactionBufferManagerState.SavedGlobalBuffer))
+	{
+		UE_LOG(
+			LogAIAssistant, Warning,
+			TEXT("RestoreGlobalBuffer: Saved global buffer was invalid, keeping current buffer"));
+		TransactionBufferManagerState.SavedGlobalBuffer = nullptr;
+		TransactionBufferManagerState.bIsOverrideActive = false;
+		return;
+	}
+	GEditor->Trans = TransactionBufferManagerState.SavedGlobalBuffer;
+	TransactionBufferManagerState.SavedGlobalBuffer = nullptr;
+	TransactionBufferManagerState.bIsOverrideActive = false;
+}
+
+bool FTransactionBufferManager::IsOverrideActive()
+{
+	check(IsInGameThread());
+
+	return TransactionBufferManagerState.bIsOverrideActive;
+}
+
+}  // namespace UE::AIAssistant
+
+// Console commands for testing and debugging
+namespace
+{
+
+void SetTransactionBufferOverrideCommand(const TArray<FString>& Args)
+{
+	using UE::AIAssistant::FTransactionBufferManager;
+
+	if (Args.Num() < 1)
+	{
+		UE_LOG(LogAIAssistant, Warning, TEXT("Usage: TransactionBuffer.SetOverride <BufferName>"));
+		return;
+	}
+	FName BufferName = *Args[0];
+	UTransBuffer* Buffer = FTransactionBufferManager::GetOrCreateTransactionBuffer(BufferName);
+	if (Buffer && FTransactionBufferManager::SetOverrideBuffer(Buffer))
+	{
+		UE_LOG(
+			LogAIAssistant, Log,
+			TEXT("Set transaction buffer override to '%s'"), *BufferName.ToString());
+	}
+	else
+	{
+		UE_LOG(LogAIAssistant, Error,
+			TEXT("Failed to set transaction buffer override to '%s'"), *BufferName.ToString());
+	}
+}
+
+void RestoreTransactionBufferCommand()
+{
+	UE::AIAssistant::FTransactionBufferManager::RestoreGlobalBuffer();
+	UE_LOG(LogAIAssistant, Log, TEXT("Restored global transaction buffer"));
+}
+
+void DestroyTransactionBufferCommand(const TArray<FString>& Args)
+{
+	using UE::AIAssistant::FTransactionBufferManager;
+
+	if (Args.Num() < 1)
+	{
+		UE_LOG(LogAIAssistant, Warning, TEXT("Usage: TransactionBuffer.Destroy <BufferName>"));
+		return;
+	}
+
+	FName BufferName(*Args[0]);
+	if (FTransactionBufferManager::DestroyTransactionBuffer(BufferName))
+	{
+		UE_LOG(
+			LogAIAssistant, Log,
+			TEXT("Destroyed transaction buffer '%s'"), *BufferName.ToString());
+	}
+	else
+	{
+		UE_LOG(
+			LogAIAssistant, Error,
+			TEXT("Failed to destroy transaction buffer '%s'"), *BufferName.ToString());
+	}
+}
+
+FAutoConsoleCommand CCmdSetTransactionBufferOverride(
+	TEXT("TransactionBuffer.SetOverride"),
+	TEXT("Set an existing transaction buffer as the active override.)")
+	TEXT("Usage: TransactionBuffer.SetOverride <BufferName>"),
+	FConsoleCommandWithArgsDelegate::CreateStatic(&SetTransactionBufferOverrideCommand)
+);
+
+FAutoConsoleCommand CCmdRestoreTransactionBuffer(
+	TEXT("TransactionBuffer.Restore"),
+	TEXT("Restore the global transaction buffer."),
+	FConsoleCommandDelegate::CreateStatic(&RestoreTransactionBufferCommand)
+);
+
+FAutoConsoleCommand CCmdDestroyTransactionBuffer(
+	TEXT("TransactionBuffer.Destroy"),
+	TEXT("Destroy a transaction buffer. Usage: TransactionBuffer.Destroy <BufferName>"),
+	FConsoleCommandWithArgsDelegate::CreateStatic(&DestroyTransactionBufferCommand)
+);
+
+} // anonymous namespace
diff --git a/Engine/Plugins/Experimental/AIAssistant/Source/AIAssistant/Private/AIAssistantTransactionBufferManager.h b/Engine/Plugins/Experimental/AIAssistant/Source/AIAssistant/Private/AIAssistantTransactionBufferManager.h
new file mode 100644
--- /dev/null
+++ b/Engine/Plugins/Experimental/AIAssistant/Source/AIAssistant/Private/AIAssistantTransactionBufferManager.h
@@ -0,0 +1,59 @@
+// Copyright Epic Games, Inc. All Rights Reserved.
+
+#pragma once
+
+#include "CoreMinimal.h"
+#include "Editor/TransBuffer.h"
+
+namespace UE::AIAssistant
+{
+
+// Manager for creating and managing multiple independent transaction buffers.
+// Allows for the creation of separate undo buffers that temporarily override the global buffer.
+class FTransactionBufferManager
+{
+public:
+	// Get an existing transaction buffer by name.
+	//
+	// @param BufferName Name of the transaction buffer to retrieve
+	// @return The transaction buffer if found, nullptr otherwise
+	static TObjectPtr<UTransBuffer> GetTransactionBuffer(const FName& BufferName);
+
+	// Create a new transaction buffer with the given name.
+	// If a buffer with this name already exists, returns the existing buffer.
+	//
+	// @param BufferName Unique name for the transaction buffer
+	// @return The created or existing transaction buffer
+	static TObjectPtr<UTransBuffer> GetOrCreateTransactionBuffer(const FName& BufferName);
+	
+	// Destroy a transaction buffer by name.
+	// Removes it from the registry and releases the root reference.
+	// Cannot destroy a buffer that is currently active as an override.
+	//
+	// @param BufferName Name of the transaction buffer to destroy
+	// @return True if the buffer was destroyed, false if not found or currently active
+	static bool DestroyTransactionBuffer(const FName& BufferName);
+
+	// Shutdown and cleanup all transaction buffers.
+	// Should be called during module shutdown.
+	// Restores any active overrides and clears all registered buffers.
+	static void Shutdown();
+
+	// Override the global undo buffer with the provided buffer.
+	// The original global undo buffer is saved on first call and can be restored at any time.
+	//
+	// @param OverrideBuffer The transaction buffer to use as override.
+	// @return True if the override was set successfully, false otherwise.
+	static bool SetOverrideBuffer(UTransactor* OverrideBuffer);
+
+	// Clears the current override and restores the global undo buffer that was saved when
+	// SetOverrideBuffer was first called.
+	static void RestoreGlobalBuffer();
+
+	// Check if there is currently an override buffer active.
+	//
+	// @return True if an override buffer is active, false otherwise
+	static bool IsOverrideActive();
+};
+
+}  // namespace UE::AIAssistant
diff --git a/Engine/Plugins/Experimental/AIAssistant/Source/AIAssistant/Private/Tests/AIAssistantTransactionBufferManagerTest.cpp b/Engine/Plugins/Experimental/AIAssistant/Source/AIAssistant/Private/Tests/AIAssistantTransactionBufferManagerTest.cpp
new file mode 100644
--- /dev/null
+++ b/Engine/Plugins/Experimental/AIAssistant/Source/AIAssistant/Private/Tests/AIAssistantTransactionBufferManagerTest.cpp
@@ -0,0 +1,288 @@
+// Copyright Epic Games, Inc. All Rights Reserved.
+
+#include "Misc/AutomationTest.h"
+
+#include "Editor.h"
+#include "Editor/EditorEngine.h"
+#include "Editor/TransBuffer.h"
+
+#include "AIAssistantTestFlags.h"
+#include "AIAssistantTransactionBufferManager.h"
+
+#if WITH_DEV_AUTOMATION_TESTS
+
+using namespace UE::AIAssistant;
+
+BEGIN_DEFINE_SPEC(FAIAssistantTransactionBufferManagerTest,
+	"AI.Assistant.TransactionBufferManager",
+	AIAssistantTest::Flags)
+	FName TestBufferName;
+	FName TestBufferName2;
+END_DEFINE_SPEC(FAIAssistantTransactionBufferManagerTest)
+
+void FAIAssistantTransactionBufferManagerTest::Define()
+{
+	if (!GEditor)
+	{
+		AddInfo(TEXT("Test skipped: No GEditor available."));
+		return;
+	}
+
+	BeforeEach([this]
+		{
+			TestBufferName = FName(TEXT("TestBuffer"));
+			TestBufferName2 = FName(TEXT("TestBuffer2"));
+		});
+
+	AfterEach([this]
+		{
+			// Clean up any buffers created during tests
+			FTransactionBufferManager::Shutdown();
+		});
+
+	Describe("GetTransactionBuffer", [this]
+		{
+			It("returns nullptr when buffer doesn't exist", [this]
+				{
+					TObjectPtr<UTransBuffer> Buffer =
+						FTransactionBufferManager::GetTransactionBuffer(TestBufferName);
+					TestNull(TEXT("Buffer should be nullptr for non-existent name"), Buffer.Get());
+				});
+
+			It("returns existing buffer after creation", [this]
+				{
+					FTransactionBufferManager::GetOrCreateTransactionBuffer(TestBufferName);
+					TObjectPtr<UTransBuffer> Buffer =
+						FTransactionBufferManager::GetTransactionBuffer(TestBufferName);
+					TestNotNull(TEXT("Buffer should be returned after creation"), Buffer.Get());
+				});
+
+			It("returns the same buffer for the same name", [this]
+				{
+					TObjectPtr<UTransBuffer> Buffer1 =
+						FTransactionBufferManager::GetOrCreateTransactionBuffer(TestBufferName);
+					TObjectPtr<UTransBuffer> Buffer2 =
+						FTransactionBufferManager::GetOrCreateTransactionBuffer(TestBufferName);
+					TestTrue(
+						TEXT("Same buffer should be returned"),
+						Buffer1.Get() == Buffer2.Get());
+				});
+		});
+
+	Describe("GetOrCreateTransactionBuffer", [this]
+		{
+			It("creates a new buffer when one doesn't exist", [this]
+				{
+					TObjectPtr<UTransBuffer> Buffer =
+						FTransactionBufferManager::GetOrCreateTransactionBuffer(TestBufferName);
+					TestNotNull(TEXT("Buffer should be created"), Buffer.Get());
+				});
+
+			It("returns the same buffer for the same name", [this]
+				{
+					TObjectPtr<UTransBuffer> Buffer1 =
+						FTransactionBufferManager::GetOrCreateTransactionBuffer(TestBufferName);
+					TObjectPtr<UTransBuffer> Buffer2 =
+						FTransactionBufferManager::GetOrCreateTransactionBuffer(TestBufferName);
+					TestTrue(
+						TEXT("Same buffer should be returned"),
+						Buffer1.Get() == Buffer2.Get());
+				});
+
+			It("returns different buffers for different names", [this]
+				{
+					TObjectPtr<UTransBuffer> Buffer1 =
+						FTransactionBufferManager::GetOrCreateTransactionBuffer(TestBufferName);
+					TObjectPtr<UTransBuffer> Buffer2 =
+						FTransactionBufferManager::GetOrCreateTransactionBuffer(TestBufferName2);
+					TestTrue(
+						TEXT("Different buffers should be returned for different names"),
+						Buffer1.Get() != Buffer2.Get());
+				});
+		});
+
+	Describe("DestroyTransactionBuffer", [this]
+		{
+			It("returns false when buffer doesn't exist", [this]
+				{
+					const bool bResult = FTransactionBufferManager::DestroyTransactionBuffer(
+						FName(TEXT("NonExistentBuffer")));
+					TestFalse(TEXT("Should return false for non-existent buffer"), bResult);
+				});
+
+			It("returns true when buffer is destroyed successfully", [this]
+				{
+					FTransactionBufferManager::GetOrCreateTransactionBuffer(TestBufferName);
+					const bool bResult =
+						FTransactionBufferManager::DestroyTransactionBuffer(TestBufferName);
+					TestTrue(TEXT("Should return true when buffer is destroyed"), bResult);
+				});
+
+			It("removes buffer from registry after destruction", [this]
+				{
+					TObjectPtr<UTransBuffer> OriginalBuffer =
+						FTransactionBufferManager::GetOrCreateTransactionBuffer(TestBufferName);
+					FTransactionBufferManager::DestroyTransactionBuffer(TestBufferName);
+					TObjectPtr<UTransBuffer> NewBuffer =
+						FTransactionBufferManager::GetOrCreateTransactionBuffer(TestBufferName);
+					TestTrue(
+						TEXT("New buffer should be created after destruction"),
+						OriginalBuffer.Get() != NewBuffer.Get());
+				});
+
+			It("cannot destroy buffer that is currently active as override", [this]
+				{
+					TObjectPtr<UTransBuffer> Buffer =
+						FTransactionBufferManager::GetOrCreateTransactionBuffer(TestBufferName);
+					FTransactionBufferManager::SetOverrideBuffer(Buffer);
+
+					const bool bResult =
+						FTransactionBufferManager::DestroyTransactionBuffer(TestBufferName);
+					TestFalse(
+						TEXT("Should return false when trying to destroy active override buffer"),
+						bResult);
+
+					FTransactionBufferManager::RestoreGlobalBuffer();
+				});
+		});
+
+	Describe("SetOverrideBuffer", [this]
+		{
+			It("returns false when OverrideBuffer is null", [this]
+				{
+					const bool bResult = FTransactionBufferManager::SetOverrideBuffer(nullptr);
+					TestFalse(TEXT("Should return false for null buffer"), bResult);
+				});
+
+			It("returns true when override is set successfully", [this]
+				{
+					TObjectPtr<UTransBuffer> Buffer =
+						FTransactionBufferManager::GetOrCreateTransactionBuffer(TestBufferName);
+					const bool bResult = FTransactionBufferManager::SetOverrideBuffer(Buffer);
+					TestTrue(TEXT("Should return true when override is set"), bResult);
+				});
+
+			It("sets the override as the active editor transaction buffer", [this]
+				{
+					TObjectPtr<UTransBuffer> Buffer =
+						FTransactionBufferManager::GetOrCreateTransactionBuffer(TestBufferName);
+					FTransactionBufferManager::SetOverrideBuffer(Buffer);
+					TestTrue(
+						TEXT("Editor Trans should be the override buffer"),
+						GEditor->Trans.Get() == Buffer.Get());
+				});
+
+			It("returns false when an override is already active", [this]
+				{
+					TObjectPtr<UTransBuffer> Buffer1 =
+						FTransactionBufferManager::GetOrCreateTransactionBuffer(TestBufferName);
+					TObjectPtr<UTransBuffer> Buffer2 =
+						FTransactionBufferManager::GetOrCreateTransactionBuffer(TestBufferName2);
+
+					const bool bFirstResult = FTransactionBufferManager::SetOverrideBuffer(Buffer1);
+					TestTrue(TEXT("First override should succeed"), bFirstResult);
+
+					const bool bSecondResult = FTransactionBufferManager::SetOverrideBuffer(Buffer2);
+					TestFalse(
+						TEXT("Second override should fail while first is active"),
+						bSecondResult);
+
+					// Verify the first override is still active
+					TestTrue(
+						TEXT("First override buffer should still be active"),
+						GEditor->Trans.Get() == Buffer1.Get());
+				});
+		});
+
+	Describe("RestoreGlobalBuffer", [this]
+		{
+			It("restores the original global buffer after override", [this]
+				{
+					UTransactor* OriginalTrans = GEditor->Trans.Get();
+					TObjectPtr<UTransBuffer> Buffer =
+						FTransactionBufferManager::GetOrCreateTransactionBuffer(TestBufferName);
+					FTransactionBufferManager::SetOverrideBuffer(Buffer);
+					FTransactionBufferManager::RestoreGlobalBuffer();
+					TestTrue(
+						TEXT("Editor Trans should be restored to original"),
+						GEditor->Trans.Get() == OriginalTrans);
+				});
+
+			It("clears the override state after restoration", [this]
+				{
+					TObjectPtr<UTransBuffer> Buffer =
+						FTransactionBufferManager::GetOrCreateTransactionBuffer(TestBufferName);
+					FTransactionBufferManager::SetOverrideBuffer(Buffer);
+					FTransactionBufferManager::RestoreGlobalBuffer();
+					TestFalse(
+						TEXT("Override should not be active after restoration"),
+						FTransactionBufferManager::IsOverrideActive());
+				});
+		});
+
+	Describe("IsOverrideActive", [this]
+		{
+			It("returns false when no override is set", [this]
+				{
+					TestFalse(
+						TEXT("Should return false when no override is active"),
+						FTransactionBufferManager::IsOverrideActive());
+				});
+
+			It("returns true when override is set", [this]
+				{
+					TObjectPtr<UTransBuffer> Buffer =
+						FTransactionBufferManager::GetOrCreateTransactionBuffer(TestBufferName);
+					FTransactionBufferManager::SetOverrideBuffer(Buffer);
+					TestTrue(
+						TEXT("Should return true when override is active"),
+						FTransactionBufferManager::IsOverrideActive());
+				});
+
+			It("returns false after override is restored", [this]
+				{
+					TObjectPtr<UTransBuffer> Buffer =
+						FTransactionBufferManager::GetOrCreateTransactionBuffer(TestBufferName);
+					FTransactionBufferManager::SetOverrideBuffer(Buffer);
+					FTransactionBufferManager::RestoreGlobalBuffer();
+					TestFalse(
+						TEXT("Should return false after override is restored"),
+						FTransactionBufferManager::IsOverrideActive());
+				});
+		});
+
+	Describe("Shutdown", [this]
+		{
+			It("clears all registered buffers", [this]
+				{
+					FTransactionBufferManager::GetOrCreateTransactionBuffer(TestBufferName);
+					FTransactionBufferManager::GetOrCreateTransactionBuffer(TestBufferName2);
+					FTransactionBufferManager::Shutdown();
+
+					// After shutdown, getting a buffer should return nullptr
+					TObjectPtr<UTransBuffer> Buffer =
+						FTransactionBufferManager::GetTransactionBuffer(TestBufferName);
+					TestNull(
+						TEXT("Buffer should be cleared after shutdown"),
+						Buffer.Get());
+				});
+
+			It("restores global buffer if override was active", [this]
+				{
+					UTransactor* OriginalTrans = GEditor->Trans.Get();
+					TObjectPtr<UTransBuffer> Buffer =
+						FTransactionBufferManager::GetOrCreateTransactionBuffer(TestBufferName);
+					FTransactionBufferManager::SetOverrideBuffer(Buffer);
+					FTransactionBufferManager::Shutdown();
+
+					TestTrue(
+						TEXT("Global buffer should be restored after shutdown"),
+						GEditor->Trans.Get() == OriginalTrans);
+					TestFalse(
+						TEXT("Override should not be active after shutdown"),
+						FTransactionBufferManager::IsOverrideActive());
+				});
+		});
+}
+
+#endif  // WITH_DEV_AUTOMATION_TESTS
\ No newline at end of file
