Added a native AIAssistantToolset which contains the functionality to figure out what asset editor the assistant is docked with.

This toolset also includes a utility function for turning graph nodes into text, but this will go away once we have a proper graph toolset.

#Synced-CL 50390180

diff --git a/Engine/Plugins/Experimental/AIAssistant/Source/AIAssistant/Private/AIAssistantSubsystem.cpp b/Engine/Plugins/Experimental/AIAssistant/Source/AIAssistant/Private/AIAssistantSubsystem.cpp
--- a/Engine/Plugins/Experimental/AIAssistant/Source/AIAssistant/Private/AIAssistantSubsystem.cpp
+++ b/Engine/Plugins/Experimental/AIAssistant/Source/AIAssistant/Private/AIAssistantSubsystem.cpp
@@ -7,10 +7,12 @@
 #include "Modules/ModuleManager.h"
 #include "Templates/SharedPointer.h"
 #include "Templates/ValueOrError.h"
+#include "ToolsetRegistry/Public/BPToolsetRegistry.h"
 #include "UObject/ObjectPtr.h"
 
 #include "AIAssistant.h"
 #include "AIAssistantLog.h"
+#include "AIAssistantToolset.h"
 #include "AIAssistantWebBrowser.h"
 
 #include UE_INLINE_GENERATED_CPP_BY_NAME(AIAssistantSubsystem)
@@ -71,10 +73,12 @@ TSharedPtr<SAIAssistantWebBrowser> UAIAssistantSubsystem::GetAIAssistantWebBrows
 void UAIAssistantSubsystem::Initialize(FSubsystemCollectionBase& Collection)
 {
 	Super::Initialize(Collection);
+	UToolsetRegistry::RegisterToolsetClass(UAIAssistantToolset::StaticClass());
 }
 
 void UAIAssistantSubsystem::Deinitialize()
 {
+	UToolsetRegistry::UnregisterToolsetClass(UAIAssistantToolset::StaticClass());
 	Super::Deinitialize();
 }
 
diff --git a/Engine/Plugins/Experimental/AIAssistant/Source/AIAssistant/Private/AIAssistantToolset.cpp b/Engine/Plugins/Experimental/AIAssistant/Source/AIAssistant/Private/AIAssistantToolset.cpp
new file mode 100644
--- /dev/null
+++ b/Engine/Plugins/Experimental/AIAssistant/Source/AIAssistant/Private/AIAssistantToolset.cpp
@@ -0,0 +1,119 @@
+// Copyright Epic Games, Inc. All Rights Reserved.
+
+#include "AIAssistantToolset.h"
+
+#include "BlueprintEditor.h"
+#include "EdGraphUtilities.h"
+#include "IMaterialEditor.h"
+#include "IPythonScriptPlugin.h"
+#include "JsonSchema/JsonSchemaGenerator.h"
+#include "MaterialGraph/MaterialGraph.h"
+#include "Materials/Material.h"
+#include "UObject/ObjectPtr.h"
+#include "Widgets/Docking/SDockTab.h"
+
+#include "AIAssistantSubsystem.h"
+#include "AIAssistantSlateQuerierUtils.h"
+#include "AIAssistantSlateQuerierUtilsTypeSearch.h"
+#include "AIAssistantWebBrowser.h"
+
+using namespace UE::AIAssistant::SlateQuerier;
+
+#include UE_INLINE_GENERATED_CPP_BY_NAME(AIAssistantToolset)
+
+
+FAIAssistantDockContext UAIAssistantToolset::GetDockedContext()
+{	
+	TSharedPtr<SAIAssistantWebBrowser> BrowserWidget = UAIAssistantSubsystem::GetAIAssistantWebBrowserWidget();
+
+	// Walk up the widget hierarchy until we encounter an asset editor.
+	TSharedPtr<SWidget> Current = BrowserWidget;
+	TSharedPtr<SWidget> StandaloneAssetEditorHost;
+	while (Current)
+	{
+		if (Current->GetTypeAsString() == "SStandaloneAssetEditorToolkitHost")
+		{
+			StandaloneAssetEditorHost = Current;
+			break;
+		}
+		Current = Current->GetParentWidget();
+	}
+
+	// If we found an asset editor, we'll try drill down through the tab system to get the asset editor.
+	if (StandaloneAssetEditorHost)
+	{
+		TArray<TSharedRef<SWidget>> TabStacks;
+		Utility::FindChildWidgetsOfType(TabStacks, StandaloneAssetEditorHost, "SDockingTabStack");
+		for (auto& ThisTabStack : TabStacks)
+		{
+			TArray<TSharedRef<SWidget>> DockTabs;
+			Utility::FindChildWidgetsOfType(DockTabs, ThisTabStack.ToSharedPtr(), "SDockTab");
+
+			for (auto& ThisTab : DockTabs)
+			{
+				TSharedRef<SDockTab> ThisTabCast = StaticCastSharedRef<SDockTab>(ThisTab);
+				if (ThisTabCast->IsForeground())
+				{
+					TSharedPtr<FTabManager> PickedWidgetTabManager = ThisTabCast->GetTabManagerPtr();
+
+					if (PickedWidgetTabManager.IsValid())
+					{
+						UAssetEditorSubsystem* AssetEditorSubsystem = GEditor->GetEditorSubsystem<UAssetEditorSubsystem>();
+						TArray<UObject*> AllEditedAssets = AssetEditorSubsystem->GetAllEditedAssets();
+						for (auto& ThisEditedAsset : AllEditedAssets)
+						{
+							IAssetEditorInstance* ThisEditor = AssetEditorSubsystem->FindEditorForAsset(ThisEditedAsset, false);
+							if (ThisEditor && ThisEditor->GetAssociatedTabManager() == PickedWidgetTabManager)
+							{
+								// This per-asset logic is not ideal, but it's the best we can do with the current AssetEditorInterface.
+								// If/when the asset editor interface is improved, we can simplify this logic.
+								FAIAssistantDockContext AssetInfo;
+								AssetInfo.Asset = ThisEditedAsset;
+								FGraphPanelSelectionSet SelectedNodes;
+								if (ThisEditedAsset->IsA<UBlueprint>())
+								{
+									FBlueprintEditor* BlueprintEditor = static_cast<FBlueprintEditor*>(ThisEditor);
+									AssetInfo.Graph = BlueprintEditor->GetFocusedGraph();
+									SelectedNodes = BlueprintEditor->GetSelectedNodes();
+								}
+								else if (UMaterial* Material = Cast<UMaterial>(ThisEditedAsset))
+								{
+									IMaterialEditor* MaterialEditor = static_cast<IMaterialEditor*>(ThisEditor);
+									AssetInfo.Graph = Material->MaterialGraph.Get();
+									SelectedNodes = MaterialEditor->GetSelectedNodes();
+								}
+								for (auto& Node : SelectedNodes)
+								{
+									if (UEdGraphNode* GraphNode = Cast<UEdGraphNode>(Node))
+									{
+										AssetInfo.SelectedNodes.Add(GraphNode);
+									}
+								}
+								return AssetInfo;
+							}
+						}
+					}
+				}
+			}
+		}
+	}
+
+	return FAIAssistantDockContext();
+}
+
+FString UAIAssistantToolset::GraphToText(const FAIAssistantDockContext& DockContext)
+{
+	FString ExportedText;
+	if (DockContext.Graph)
+	{
+		TArray<TObjectPtr<UEdGraphNode>> SelectedNodes = DockContext.SelectedNodes.Num() > 0 ? DockContext.SelectedNodes : DockContext.Graph->Nodes;
+
+		TSet<UObject*> NodesToExport;
+		for (const auto& Node : SelectedNodes)
+		{
+			NodesToExport.Add(Node.Get());
+		}		
+		FEdGraphUtilities::ExportNodesToText(NodesToExport, ExportedText);
+	}
+	return ExportedText;
+}
diff --git a/Engine/Plugins/Experimental/AIAssistant/Source/AIAssistant/Private/AIAssistantToolset.h b/Engine/Plugins/Experimental/AIAssistant/Source/AIAssistant/Private/AIAssistantToolset.h
new file mode 100644
--- /dev/null
+++ b/Engine/Plugins/Experimental/AIAssistant/Source/AIAssistant/Private/AIAssistantToolset.h
@@ -0,0 +1,48 @@
+// Copyright Epic Games, Inc. All Rights Reserved. 
+
+#pragma once
+
+#include "Kismet/BlueprintFunctionLibrary.h"
+
+#include "AIAssistantToolset.generated.h"
+
+class UEdGraph;
+class UEdGraphNode;
+
+/// Contains information about the state of the editor that the assistant is currently docked with.
+USTRUCT(BlueprintType)
+struct FAIAssistantDockContext
+{
+	GENERATED_USTRUCT_BODY()
+
+public:
+	/// The asset instance that is currently being edited (if any).
+	UPROPERTY(BlueprintReadOnly, Category = "AIAssistant")
+	TObjectPtr<UObject> Asset;
+
+	/// The graph instance that is currently being edited (if any).
+	UPROPERTY(BlueprintReadOnly, Category = "AIAssistant")
+	TObjectPtr<UEdGraph> Graph;	
+
+	/// The list of selected graph nodes edited (if any).
+	UPROPERTY(BlueprintReadOnly, Category = "AIAssistant")
+	TArray<TObjectPtr<UEdGraphNode>> SelectedNodes;
+};
+
+/// These functions help the AI assistant inspect its state as it relates to the larger Unreal Editor.
+UCLASS(BlueprintType)
+class UAIAssistantToolset : public UBlueprintFunctionLibrary
+{
+	GENERATED_BODY()
+public:
+	/// Returns information about the part of the unreal UI that the assistant is docked into.	
+	/// For example, if the assistant is docked with a Blueprint asset editor, this return
+	/// the Blueprint asset that is being edited, which graph is in focus, etc.
+	UFUNCTION(BlueprintCallable, Category = "AI Assistant")
+	static FAIAssistantDockContext GetDockedContext();
+
+	/// This converts graph information in a dock context into a string that describes it.
+	/// This string will represent the entire graph if SelectedNodes is empty.
+	UFUNCTION(BlueprintCallable, Category = "AI Assistant")
+	static FString GraphToText(const FAIAssistantDockContext& DockContext);
+};
