[AI Assistant] Fix null reference on web application shutdown

#rnx

#Synced-CL 50436129

diff --git a/Engine/Plugins/Experimental/AIAssistant/Source/AIAssistant/Private/AIAssistantWebApi.cpp b/Engine/Plugins/Experimental/AIAssistant/Source/AIAssistant/Private/AIAssistantWebApi.cpp
--- a/Engine/Plugins/Experimental/AIAssistant/Source/AIAssistant/Private/AIAssistantWebApi.cpp
+++ b/Engine/Plugins/Experimental/AIAssistant/Source/AIAssistant/Private/AIAssistantWebApi.cpp
@@ -162,6 +162,12 @@ try {
 	void FWebApi::UnregisterOnConversationUpdate(const FConversationUpdateHandle& Handle)
 	{
 		auto& [HandlerId, CallbackHandle] = Handle;
+		// Unfortunately on shutdown the delegate can be garbage collected even though a strong
+		// reference to it is held.
+		if (!WebJavaScriptResultDelegate.IsValid())
+		{
+			return;
+		}
 		(void)ExecuteFunction(
 			*WebApiObjectName, TEXT("unregisterOnConversationUpdate"),
 			*FString::Printf(TEXT("%s"), *CallbackHandle));
@@ -246,4 +252,4 @@ try {
 		WebJavaScriptResultDelegate.Reset();
 	}
 
-}  // namespace UE::AIAssistant
\ No newline at end of file
+}  // namespace UE::AIAssistant
