Fix heap-use-after-free in unit tests.

The ToolsetRegistry needs to outlive the WebApplication, and when the ToolsetRegistry is a member of the WebApplication that assumption no longer holds true. For the tests that use the SpyWebApplication, the ToolsetRegistry's lifetime was adjusted to live at least as long as the WebApplication.

#Synced-CL 50423628

diff --git a/Engine/Plugins/Experimental/AIAssistant/Source/AIAssistant/Private/Tests/AIAssistantWebApplicationTest.cpp b/Engine/Plugins/Experimental/AIAssistant/Source/AIAssistant/Private/Tests/AIAssistantWebApplicationTest.cpp
--- a/Engine/Plugins/Experimental/AIAssistant/Source/AIAssistant/Private/Tests/AIAssistantWebApplicationTest.cpp
+++ b/Engine/Plugins/Experimental/AIAssistant/Source/AIAssistant/Private/Tests/AIAssistantWebApplicationTest.cpp
@@ -240,10 +240,15 @@ namespace UE::AIAssistant
 
 		virtual ~FSpyWebApplication() override = default;
 
-		UE::ToolsetRegistry::FToolsetRegistry Registry;
-		UE::ToolsetRegistry::FToolsetRegistry& GetToolsetRegistry() override
+		void SetRegistry(TSharedPtr<UE::ToolsetRegistry::FToolsetRegistry> InRegistry)
 		{
-			return Registry;
+			Registry = InRegistry;
+		}
+
+		virtual UE::ToolsetRegistry::FToolsetRegistry& GetToolsetRegistry() override
+		{
+			check(Registry);
+			return *Registry;
 		}
 
 		using FWebApplication::OnConversationUpdate;
@@ -280,6 +285,8 @@ namespace UE::AIAssistant
 			return ResultPromise->GetFuture();
 		}
 
+		TSharedPtr<UE::ToolsetRegistry::FToolsetRegistry> Registry;
+
 		bool bCallBaseAddUserMessageToConversation = true;
 		TArray<FAddMessageToConversationOptions> AddUserMessageToConversationCalls;
 
@@ -802,7 +809,9 @@ namespace UE::AIAssistant
 	bool FAIAssistantWebApplicationOnConversationUpdateProcessMessages::RunTest(
 		const FString& UnusedParameters)
 	{
+		auto Registry = MakeShared<UE::ToolsetRegistry::FToolsetRegistry>();
 		FSpyWebApplicationWithTracker WebApplicationWithTracker(*this);
+		WebApplicationWithTracker.WebApplication->SetRegistry(Registry);
 		auto WebApplication = WebApplicationWithTracker.WebApplication;
 		if (!WebApplicationWithTracker.EnsureInitialized()) return false;
 
@@ -849,7 +858,9 @@ namespace UE::AIAssistant
 	bool FAIAssistantWebApplicationProcessMessageTextContentMessage::RunTest(
 		const FString& UnusedParameters)
 	{
+		auto Registry = MakeShared<UE::ToolsetRegistry::FToolsetRegistry>();
 		FSpyWebApplicationWithTracker WebApplicationWithTracker(*this);
+		WebApplicationWithTracker.WebApplication->SetRegistry(Registry);
 		auto WebApplication = WebApplicationWithTracker.WebApplication;
 		if (!WebApplicationWithTracker.EnsureInitialized()) return false;
 
@@ -871,7 +882,9 @@ namespace UE::AIAssistant
 	bool FAIAssistantWebApplicationProcessMessageToolCallMessage::RunTest(
 		const FString& UnusedParameters)
 	{
+		auto Registry = MakeShared<UE::ToolsetRegistry::FToolsetRegistry>();
 		FSpyWebApplicationWithTracker WebApplicationWithTracker(*this);
+		WebApplicationWithTracker.WebApplication->SetRegistry(Registry);
 		auto WebApplication = WebApplicationWithTracker.WebApplication;
 		if (!WebApplicationWithTracker.EnsureInitialized()) return false;
 
@@ -921,7 +934,9 @@ namespace UE::AIAssistant
 	bool FAIAssistantWebApplicationProcessMessageOnlyProcessRequiredToolCalls::RunTest(
 		const FString& UnusedParameters)
 	{
+		auto Registry = MakeShared<UE::ToolsetRegistry::FToolsetRegistry>();
 		FSpyWebApplicationWithTracker WebApplicationWithTracker(*this);
+		WebApplicationWithTracker.WebApplication->SetRegistry(Registry);
 		auto WebApplication = WebApplicationWithTracker.WebApplication;
 		if (!WebApplicationWithTracker.EnsureInitialized()) return false;
 
@@ -968,7 +983,9 @@ namespace UE::AIAssistant
 	bool FAIAssistantWebApplicationProcessMessageProcessOnlyOnce::RunTest(
 		const FString& UnusedParameters)
 	{
+		auto Registry = MakeShared<UE::ToolsetRegistry::FToolsetRegistry>();
 		FSpyWebApplicationWithTracker WebApplicationWithTracker(*this);
+		WebApplicationWithTracker.WebApplication->SetRegistry(Registry);
 		auto WebApplication = WebApplicationWithTracker.WebApplication;
 		if (!WebApplicationWithTracker.EnsureInitialized()) return false;
 
@@ -1014,7 +1031,9 @@ namespace UE::AIAssistant
 	bool FAIAssistantWebApplicationProcessMessageExitGracefullyOnError::RunTest(
 		const FString& UnusedParameters)
 	{
+		auto Registry = MakeShared<UE::ToolsetRegistry::FToolsetRegistry>();
 		FSpyWebApplicationWithTracker WebApplicationWithTracker(*this);
+		WebApplicationWithTracker.WebApplication->SetRegistry(Registry);
 		auto WebApplication = WebApplicationWithTracker.WebApplication;
 		if (!WebApplicationWithTracker.EnsureInitialized()) return false;
 
@@ -1065,7 +1084,9 @@ namespace UE::AIAssistant
 	bool FAIAssistantWebApplicationProcessMessageExitGracefullyOnShutdown::RunTest(
 		const FString& UnusedParameters)
 	{
+		auto Registry = MakeShared<UE::ToolsetRegistry::FToolsetRegistry>();
 		FSpyWebApplicationWithTracker WebApplicationWithTracker(*this);
+		WebApplicationWithTracker.WebApplication->SetRegistry(Registry);
 		if (!WebApplicationWithTracker.EnsureInitialized()) return false;
 
 		// Messages with tool calls and valid IDs.
@@ -1102,7 +1123,9 @@ namespace UE::AIAssistant
 	bool FAIAssistantWebApplicationProcessToolCallContentSuccessfulCall::RunTest(
 		const FString& UnusedParameters)
 	{
+		auto Registry = MakeShared<UE::ToolsetRegistry::FToolsetRegistry>();
 		FSpyWebApplicationWithTracker WebApplicationWithTracker(*this);
+		WebApplicationWithTracker.WebApplication->SetRegistry(Registry);
 		auto WebApplication = WebApplicationWithTracker.WebApplication;
 		if (!WebApplicationWithTracker.EnsureInitialized()) return false;
 
@@ -1113,10 +1136,9 @@ namespace UE::AIAssistant
 			MakeTuple(1, 1));
 
 		// Set up dummy toolset.
-		auto& Registry = WebApplication->GetToolsetRegistry();
 		auto SandwichToolset = MakeShared<UE::ToolsetRegistry::FFakeToolset>("Sandwich");
 		auto PeanutButterToolCall = SandwichToolset->AddFakeToolCall(TEXT("SpreadPeanutButter"));
-		Registry.RegisterToolset(SandwichToolset);
+		Registry->RegisterToolset(SandwichToolset);
 		WebApplication->bCallBaseAddUserMessageToConversation = false;
 
 		// Process tool call and fulfill the promise successfully.
@@ -1152,7 +1174,9 @@ namespace UE::AIAssistant
 	bool FAIAssistantWebApplicationProcessToolCallContentErroneousCall::RunTest(
 		const FString& UnusedParameters)
 	{
+		auto Registry = MakeShared<UE::ToolsetRegistry::FToolsetRegistry>();
 		FSpyWebApplicationWithTracker WebApplicationWithTracker(*this);
+		WebApplicationWithTracker.WebApplication->SetRegistry(Registry);
 		auto WebApplication = WebApplicationWithTracker.WebApplication;
 		if (!WebApplicationWithTracker.EnsureInitialized()) return false;
 
@@ -1163,11 +1187,10 @@ namespace UE::AIAssistant
 			MakeTuple(1, 1));
 
 		// Set up dummy toolset.
-		auto& Registry = WebApplication->GetToolsetRegistry();
 		auto SandwichToolset = MakeShared<UE::ToolsetRegistry::FFakeToolset>("Sandwich");
 		auto PeanutButterToolCallPromise =
 			SandwichToolset->AddFakeToolCall(TEXT("SpreadPeanutButter"));
-		Registry.RegisterToolset(SandwichToolset);
+		Registry->RegisterToolset(SandwichToolset);
 
 		// Process tool call and wait for it to complete.
 		WebApplication->bCallBaseAddUserMessageToConversation = false;
