Add a file lock manager to prevent asset editing in the editor.

Added debug commands to add or remove files from the set of locked files. The lock manager hooks into the AssetEditorSubsystem's filters to decide whether a file should be editable.

#Synced-CL 50423628

diff --git a/Engine/Plugins/Experimental/AIAssistant/Source/AIAssistant/Private/AIAssistant.cpp b/Engine/Plugins/Experimental/AIAssistant/Source/AIAssistant/Private/AIAssistant.cpp
--- a/Engine/Plugins/Experimental/AIAssistant/Source/AIAssistant/Private/AIAssistant.cpp
+++ b/Engine/Plugins/Experimental/AIAssistant/Source/AIAssistant/Private/AIAssistant.cpp
@@ -19,6 +19,7 @@
 #include "AIAssistantSlateQuerier.h"
 #include "AIAssistantStyle.h"
 #include "AIAssistantWebBrowser.h"
+#include "AIAssistantFileLockManager.h"
 
 #define LOCTEXT_NAMESPACE "FAIAssistantModule"
 
@@ -80,10 +81,15 @@ static const FName AIAssistantTabName("AIAssistant");
 	if (GEditor)
 	{
 		RegisterStatusBarPanelDrawerSummon();
+		UE::AIAssistant::FFileLockManager::RegisterWithEditor();
 	}
 	else
 	{
 		FCoreDelegates::OnPostEngineInit.AddRaw(this, &FAIAssistantModule::RegisterStatusBarPanelDrawerSummon);
+		FCoreDelegates::OnPostEngineInit.AddLambda([]()
+		{
+			UE::AIAssistant::FFileLockManager::RegisterWithEditor();
+		});
 	}
 
 	// Some scenarios, like commandlets, don't have slate initialized
@@ -114,6 +120,8 @@ static const FName AIAssistantTabName("AIAssistant");
 	FAIAssistantStyle::Shutdown();
 
 	FAIAssistantCommands::Unregister();
+	
+	UE::AIAssistant::FFileLockManager::UnregisterFromEditor();
 
 	FGlobalTabmanager::Get()->UnregisterNomadTabSpawner(AIAssistantTabName);
 }
diff --git a/Engine/Plugins/Experimental/AIAssistant/Source/AIAssistant/Private/AIAssistantFileLockManager.cpp b/Engine/Plugins/Experimental/AIAssistant/Source/AIAssistant/Private/AIAssistantFileLockManager.cpp
new file mode 100644
--- /dev/null
+++ b/Engine/Plugins/Experimental/AIAssistant/Source/AIAssistant/Private/AIAssistantFileLockManager.cpp
@@ -0,0 +1,320 @@
+// Copyright Epic Games, Inc. All Rights Reserved.
+
+#include "AIAssistantFileLockManager.h"
+
+#include "AIAssistantLog.h"
+#include "Editor.h"
+#include "Containers/Ticker.h"
+#include "AssetTypeActivationOpenedMethod.h"
+#include "HAL/IConsoleManager.h"
+#include "Subsystems/AssetEditorSubsystem.h"
+
+namespace
+{
+	// Encapsulates all mutable state for the file lock manager.
+	struct FFileLockManagerState
+	{
+		// Mutex protecting LockedFiles.
+		UE::FMutex Mutex;
+
+		// Set of locked package paths.
+		TSet<FString> LockedFiles;
+
+		// Handle for the OnAssetOpenedInEditor delegate.
+		FDelegateHandle AssetOpenedHandle;
+
+		void Reset()
+		{
+			UE::TUniqueLock Lock(Mutex);
+			LockedFiles.Empty();
+			AssetOpenedHandle.Reset();
+		}
+	};
+
+	static FFileLockManagerState FileLockManagerState;
+
+	static FString NormalizePath(const FString& PackagePath)
+	{
+		FString NormalizedPath = PackagePath.TrimStartAndEnd();
+		if (NormalizedPath.IsEmpty())
+		{
+			return NormalizedPath;
+		}
+		if (!NormalizedPath.StartsWith(TEXT("/")))
+		{
+			NormalizedPath = TEXT("/") + NormalizedPath;
+		}
+		return NormalizedPath;
+	}
+
+	static void OnAssetOpenedInEditor(UObject* Asset, IAssetEditorInstance* EditorInstance)
+	{
+		using UE::AIAssistant::FFileLockManager;
+
+		if (!Asset || !EditorInstance)
+		{
+			return;
+		}
+
+		UPackage* Package = Asset->GetPackage();
+		if (!Package || !FFileLockManager::IsFileLocked(Package->GetName()))
+		{
+			return;
+		}
+
+		// If already in View mode, nothing to do.
+		if (EditorInstance->GetOpenMethod() == EAssetOpenMethod::View)
+		{
+			return;
+		}
+
+		// Asset is locked but opened for editing: defer close and reopen in View mode
+		// to avoid modifying editor state during the notification callback.
+		FTSTicker::GetCoreTicker().AddTicker(
+			FTickerDelegate::CreateLambda(
+				[WeakAsset = TWeakObjectPtr<UObject>(Asset)](float) -> bool
+				{
+					UObject* AssetPtr = WeakAsset.Get();
+					if (!GEditor || !AssetPtr)
+					{
+						return false;
+					}
+					UAssetEditorSubsystem* Subsystem =
+						GEditor->GetEditorSubsystem<UAssetEditorSubsystem>();
+					if (!Subsystem)
+					{
+						return false;
+					}
+					// Close and reopen the asset in View mode.
+					// If the user has edits to the file they will not be lost, they will remain
+					// in-memory and be reflected when the asset editor is reopened.
+					Subsystem->CloseAllEditorsForAsset(AssetPtr);
+					bool bShowProgressWindow = true;
+					Subsystem->OpenEditorForAsset(
+						AssetPtr, EToolkitMode::Standalone, TSharedPtr<IToolkitHost>(),
+						bShowProgressWindow, EAssetTypeActivationOpenedMethod::View);
+					
+					return false;
+				}),
+			0.0f);
+	}
+}  // namespace
+
+namespace UE::AIAssistant
+{
+
+bool FFileLockManager::AddLockedFile(const FString& PackagePath)
+{
+	FString NormalizedPath = NormalizePath(PackagePath);
+	if (NormalizedPath.IsEmpty())
+	{
+		return false;
+	}
+
+	UE::TUniqueLock Lock(FileLockManagerState.Mutex);
+	bool bAlreadyInSet = false;
+	FileLockManagerState.LockedFiles.Add(NormalizedPath, &bAlreadyInSet);
+	return !bAlreadyInSet;
+}
+
+bool FFileLockManager::RemoveLockedFile(const FString& PackagePath)
+{
+	FString NormalizedPath = NormalizePath(PackagePath);
+	if (NormalizedPath.IsEmpty())
+	{
+		return false;
+	}
+
+	UE::TUniqueLock Lock(FileLockManagerState.Mutex);
+	return FileLockManagerState.LockedFiles.Remove(NormalizedPath) > 0;
+}
+
+bool FFileLockManager::IsFileLocked(const FString& PackagePath)
+{
+	FString NormalizedPath = NormalizePath(PackagePath);
+
+	UE::TUniqueLock Lock(FileLockManagerState.Mutex);
+
+	// Check exact match and prefix match (for sub-assets)
+	for (const FString& LockedPath : FileLockManagerState.LockedFiles)
+	{
+		if (NormalizedPath == LockedPath || NormalizedPath.StartsWith(LockedPath + TEXT(".")))
+		{
+			return true;
+		}
+	}
+
+	return false;
+}
+
+TSet<FString> FFileLockManager::GetLockedFiles()
+{
+	UE::TUniqueLock Lock(FileLockManagerState.Mutex);
+	return FileLockManagerState.LockedFiles;
+}
+
+void FFileLockManager::ClearLockedFiles()
+{
+	UE::TUniqueLock Lock(FileLockManagerState.Mutex);
+	FileLockManagerState.LockedFiles.Empty();
+}
+
+void FFileLockManager::RegisterWithEditor()
+{
+	if (!GEditor)
+	{
+		UE_LOG(LogAIAssistant, Warning, TEXT("RegisterWithEditor called without a valid GEditor"));
+		return;
+	}
+
+	if (FileLockManagerState.AssetOpenedHandle.IsValid())
+	{
+		return;
+	}
+	UAssetEditorSubsystem* AssetEditorSubsystem =
+		GEditor->GetEditorSubsystem<UAssetEditorSubsystem>();
+	if (!AssetEditorSubsystem)
+	{
+		UE_LOG(LogAIAssistant, Warning, TEXT("Could not access AssetEditorSubsystem"));
+		return;
+	}
+
+	// When an asset opens in an editor, check if it's locked and force View mode
+	FileLockManagerState.AssetOpenedHandle =
+		AssetEditorSubsystem->OnAssetOpenedInEditor().AddStatic(&OnAssetOpenedInEditor);
+}
+
+void FFileLockManager::UnregisterFromEditor()
+{
+	if (!GEditor)
+	{
+		return;
+	}
+
+	UAssetEditorSubsystem* AssetEditorSubsystem = GEditor->GetEditorSubsystem<UAssetEditorSubsystem>();
+	if (AssetEditorSubsystem)
+	{
+		AssetEditorSubsystem->OnAssetOpenedInEditor().Remove(FileLockManagerState.AssetOpenedHandle);
+	}
+
+	FileLockManagerState.Reset();
+}
+
+}  // namespace UE::AIAssistant
+
+// Console commands for testing and debugging
+namespace
+{
+
+void LockFileCommand(const TArray<FString>& Args)
+{
+	using UE::AIAssistant::FFileLockManager;
+
+	if (Args.Num() == 0)
+	{
+		UE_LOG(
+			LogAIAssistant, Warning, 
+			TEXT("AIAssistant.LockFile: Please provide a package path. "
+				 "Usage: AIAssistant.LockFile <PackagePath>"));
+		return;
+	}
+
+	FString PackagePath = Args[0];
+
+	if (FFileLockManager::AddLockedFile(PackagePath))
+	{
+		UE_LOG(
+			LogAIAssistant, Log,
+			TEXT("AIAssistant.LockFile: Successfully locked file: %s"), *PackagePath);
+	}
+	else
+	{
+		UE_LOG(
+			LogAIAssistant, Warning,
+			TEXT("AIAssistant.LockFile: File is already locked: %s"), *PackagePath);
+	}
+}
+
+void UnlockFileCommand(const TArray<FString>& Args)
+{
+	using UE::AIAssistant::FFileLockManager;
+
+	if (Args.Num() == 0)
+	{
+		UE_LOG(
+			LogAIAssistant, Warning,
+			TEXT("AIAssistant.UnlockFile: Please provide a package path. "
+				 "Usage: AIAssistant.UnlockFile <PackagePath>"));
+		return;
+	}
+
+	FString PackagePath = Args[0];
+
+	if (FFileLockManager::RemoveLockedFile(PackagePath))
+	{
+		UE_LOG(
+			LogAIAssistant, Log,
+			TEXT("AIAssistant.UnlockFile: Successfully unlocked file: %s"), *PackagePath);
+	}
+	else
+	{
+		UE_LOG(
+			LogAIAssistant, Warning,
+			TEXT("AIAssistant.UnlockFile: File was not locked: %s"), *PackagePath);
+	}
+}
+
+void ListLockedFilesCommand()
+{
+	const TSet<FString> LockedFiles = UE::AIAssistant::FFileLockManager::GetLockedFiles();
+
+	if (LockedFiles.Num() == 0)
+	{
+		UE_LOG(
+			LogAIAssistant, Log,
+			TEXT("AIAssistant.ListLockedFiles: No files are currently locked."));
+	}
+	else
+	{
+		UE_LOG(
+			LogAIAssistant, Log,
+			TEXT("AIAssistant.ListLockedFiles: %d file(s) are locked:"), LockedFiles.Num());
+		for (const FString& FilePath : LockedFiles)
+		{
+			UE_LOG(LogAIAssistant, Log, TEXT("  - %s"), *FilePath);
+		}
+	}
+}
+
+void ClearLockedFilesCommand()
+{
+	UE::AIAssistant::FFileLockManager::ClearLockedFiles();
+	UE_LOG(
+		LogAIAssistant, Log, TEXT("AIAssistant.ClearLockedFiles: All locked files have been cleared."));
+}
+
+FAutoConsoleCommand CCmdLockFile(
+	TEXT("AIAssistant.LockFile"),
+	TEXT("Adds a file to the locked files list. Usage: AIAssistant.LockFile <PackagePath>"),
+	FConsoleCommandWithArgsDelegate::CreateStatic(&LockFileCommand)
+);
+
+FAutoConsoleCommand CCmdUnlockFile(
+	TEXT("AIAssistant.UnlockFile"),
+	TEXT("Removes a file from the locked files list. Usage: AIAssistant.UnlockFile <PackagePath>"),
+	FConsoleCommandWithArgsDelegate::CreateStatic(&UnlockFileCommand)
+);
+
+FAutoConsoleCommand CCmdListLockedFiles(
+	TEXT("AIAssistant.ListLockedFiles"),
+	TEXT("Lists all currently locked files."),
+	FConsoleCommandDelegate::CreateStatic(&ListLockedFilesCommand)
+);
+
+FAutoConsoleCommand CCmdClearLockedFiles(
+	TEXT("AIAssistant.ClearLockedFiles"),
+	TEXT("Clears all locked files from the list."),
+	FConsoleCommandDelegate::CreateStatic(&ClearLockedFilesCommand)
+);
+
+}  // namespace
diff --git a/Engine/Plugins/Experimental/AIAssistant/Source/AIAssistant/Private/AIAssistantFileLockManager.h b/Engine/Plugins/Experimental/AIAssistant/Source/AIAssistant/Private/AIAssistantFileLockManager.h
new file mode 100644
--- /dev/null
+++ b/Engine/Plugins/Experimental/AIAssistant/Source/AIAssistant/Private/AIAssistantFileLockManager.h
@@ -0,0 +1,37 @@
+// Copyright Epic Games, Inc. All Rights Reserved.
+
+#pragma once
+
+#include "CoreMinimal.h"
+
+namespace UE::AIAssistant
+{
+
+// Manages a list of files that should be locked from editing in the Unreal Editor.
+// Files in this list will be opened in read-only mode.
+class FFileLockManager
+{
+public:
+	// Add a file to the locked files list. Returns true if added, false if already present.
+	static bool AddLockedFile(const FString& PackagePath);
+
+	// Remove a file from the locked files list. Returns true if removed, false if not found.
+	static bool RemoveLockedFile(const FString& PackagePath);
+
+	// Check if a file is locked by package path (e.g., "/Game/MyAsset").
+	static bool IsFileLocked(const FString& PackagePath);
+
+	// Get all locked files.
+	static TSet<FString> GetLockedFiles();
+
+	// Clear all locked files.
+	static void ClearLockedFiles();
+
+	// Register file locking hooks with the editor (asset open callback and read-only filter).
+	static void RegisterWithEditor();
+
+	// Unregister file locking hooks from the editor.
+	static void UnregisterFromEditor();
+};
+
+}  // namespace UE::AIAssistant
diff --git a/Engine/Plugins/Experimental/AIAssistant/Source/AIAssistant/Private/Tests/AIAssistantFileLockManagerTest.cpp b/Engine/Plugins/Experimental/AIAssistant/Source/AIAssistant/Private/Tests/AIAssistantFileLockManagerTest.cpp
new file mode 100644
--- /dev/null
+++ b/Engine/Plugins/Experimental/AIAssistant/Source/AIAssistant/Private/Tests/AIAssistantFileLockManagerTest.cpp
@@ -0,0 +1,177 @@
+// Copyright Epic Games, Inc. All Rights Reserved.
+
+#include "Misc/AutomationTest.h"
+
+#include "AIAssistantTestFlags.h"
+#include "AIAssistantFileLockManager.h"
+
+#if WITH_DEV_AUTOMATION_TESTS
+
+BEGIN_DEFINE_SPEC(FAIAssistantFileLockManagerTest,
+	"AI.Assistant.FileLockManager",
+	AIAssistantTest::Flags)
+END_DEFINE_SPEC(FAIAssistantFileLockManagerTest)
+
+void FAIAssistantFileLockManagerTest::Define()
+{
+	using namespace UE::AIAssistant;
+
+	AfterEach([this]
+		{
+			FFileLockManager::ClearLockedFiles();
+		});
+
+	Describe("AddLockedFile", [this]
+		{
+			It("returns true when adding a new file", [this]
+				{
+					const bool bResult = FFileLockManager::AddLockedFile(TEXT("/Game/MyAsset"));
+					TestTrue(TEXT("Should return true for new file"), bResult);
+				});
+
+			It("returns false when adding a duplicate file", [this]
+				{
+					FFileLockManager::AddLockedFile(TEXT("/Game/MyAsset"));
+					const bool bResult = FFileLockManager::AddLockedFile(TEXT("/Game/MyAsset"));
+					TestFalse(TEXT("Should return false for duplicate"), bResult);
+				});
+
+			It("returns false for empty path", [this]
+				{
+					const bool bResult = FFileLockManager::AddLockedFile(TEXT(""));
+					TestFalse(TEXT("Should return false for empty path"), bResult);
+				});
+
+			It("normalizes paths without leading slash", [this]
+				{
+					FFileLockManager::AddLockedFile(TEXT("Game/MyAsset"));
+					TestTrue(
+						TEXT("Should be findable with leading slash"),
+						FFileLockManager::IsFileLocked(TEXT("/Game/MyAsset")));
+				});
+
+			It("treats paths with and without leading slash as the same", [this]
+				{
+					FFileLockManager::AddLockedFile(TEXT("/Game/MyAsset"));
+					const bool bResult = FFileLockManager::AddLockedFile(TEXT("Game/MyAsset"));
+					TestFalse(TEXT("Should detect duplicate after normalization"), bResult);
+				});
+		});
+
+	Describe("RemoveLockedFile", [this]
+		{
+			It("returns true when removing an existing file", [this]
+				{
+					FFileLockManager::AddLockedFile(TEXT("/Game/MyAsset"));
+					const bool bResult = FFileLockManager::RemoveLockedFile(TEXT("/Game/MyAsset"));
+					TestTrue(TEXT("Should return true for existing file"), bResult);
+				});
+
+			It("returns false when removing a non-existent file", [this]
+				{
+					const bool bResult = FFileLockManager::RemoveLockedFile(TEXT("/Game/MyAsset"));
+					TestFalse(TEXT("Should return false for non-existent file"), bResult);
+				});
+
+			It("file is no longer locked after removal", [this]
+				{
+					FFileLockManager::AddLockedFile(TEXT("/Game/MyAsset"));
+					FFileLockManager::RemoveLockedFile(TEXT("/Game/MyAsset"));
+					TestFalse(
+						TEXT("File should not be locked after removal"),
+						FFileLockManager::IsFileLocked(TEXT("/Game/MyAsset")));
+				});
+
+			It("normalizes paths without leading slash", [this]
+				{
+					FFileLockManager::AddLockedFile(TEXT("/Game/MyAsset"));
+					const bool bResult = FFileLockManager::RemoveLockedFile(TEXT("Game/MyAsset"));
+					TestTrue(TEXT("Should remove via normalized path"), bResult);
+				});
+		});
+
+	Describe("IsFileLocked", [this]
+		{
+			It("returns false for unlocked file", [this]
+				{
+					TestFalse(
+						TEXT("Should return false for unlocked file"),
+						FFileLockManager::IsFileLocked(TEXT("/Game/MyAsset")));
+				});
+
+			It("returns true for locked file", [this]
+				{
+					FFileLockManager::AddLockedFile(TEXT("/Game/MyAsset"));
+					TestTrue(
+						TEXT("Should return true for locked file"),
+						FFileLockManager::IsFileLocked(TEXT("/Game/MyAsset")));
+				});
+
+			It("matches sub-assets via dot prefix", [this]
+				{
+					FFileLockManager::AddLockedFile(TEXT("/Game/MyAsset"));
+					TestTrue(
+						TEXT("Sub-asset should be locked"),
+						FFileLockManager::IsFileLocked(TEXT("/Game/MyAsset.SubObject")));
+				});
+
+			It("does not match unrelated paths that share a prefix", [this]
+				{
+					FFileLockManager::AddLockedFile(TEXT("/Game/My"));
+					TestFalse(
+						TEXT("Should not match path that merely starts with the locked prefix"),
+						FFileLockManager::IsFileLocked(TEXT("/Game/MyAsset")));
+				});
+
+			It("normalizes paths without leading slash", [this]
+				{
+					FFileLockManager::AddLockedFile(TEXT("/Game/MyAsset"));
+					TestTrue(
+						TEXT("Should find locked file without leading slash"),
+						FFileLockManager::IsFileLocked(TEXT("Game/MyAsset")));
+				});
+		});
+
+	Describe("GetLockedFiles", [this]
+		{
+			It("returns empty set when no files are locked", [this]
+				{
+					const TSet<FString> LockedFiles = FFileLockManager::GetLockedFiles();
+					TestEqual(TEXT("Should have no locked files"), LockedFiles.Num(), 0);
+				});
+
+			It("returns all locked files", [this]
+				{
+					FFileLockManager::AddLockedFile(TEXT("/Game/AssetA"));
+					FFileLockManager::AddLockedFile(TEXT("/Game/AssetB"));
+					const TSet<FString> LockedFiles = FFileLockManager::GetLockedFiles();
+					TestEqual(TEXT("Should have two locked files"), LockedFiles.Num(), 2);
+					TestTrue(TEXT("Should contain AssetA"), LockedFiles.Contains(TEXT("/Game/AssetA")));
+					TestTrue(TEXT("Should contain AssetB"), LockedFiles.Contains(TEXT("/Game/AssetB")));
+				});
+		});
+
+	Describe("ClearLockedFiles", [this]
+		{
+			It("removes all locked files", [this]
+				{
+					FFileLockManager::AddLockedFile(TEXT("/Game/AssetA"));
+					FFileLockManager::AddLockedFile(TEXT("/Game/AssetB"));
+					FFileLockManager::ClearLockedFiles();
+					TestEqual(
+						TEXT("Should have no locked files after clear"),
+						FFileLockManager::GetLockedFiles().Num(), 0);
+				});
+
+			It("previously locked files are no longer locked", [this]
+				{
+					FFileLockManager::AddLockedFile(TEXT("/Game/MyAsset"));
+					FFileLockManager::ClearLockedFiles();
+					TestFalse(
+						TEXT("File should not be locked after clear"),
+						FFileLockManager::IsFileLocked(TEXT("/Game/MyAsset")));
+				});
+		});
+}
+
+#endif  // WITH_DEV_AUTOMATION_TESTS
