Changed FName to FString and added function to list files in an undo stack

#Synced-CL 50702752

diff --git a/Engine/Plugins/Experimental/AIAssistant/Source/AIAssistant/Private/AIAssistantTransactionBufferManager.cpp b/Engine/Plugins/Experimental/AIAssistant/Source/AIAssistant/Private/AIAssistantTransactionBufferManager.cpp
--- a/Engine/Plugins/Experimental/AIAssistant/Source/AIAssistant/Private/AIAssistantTransactionBufferManager.cpp
+++ b/Engine/Plugins/Experimental/AIAssistant/Source/AIAssistant/Private/AIAssistantTransactionBufferManager.cpp
@@ -23,7 +23,7 @@ namespace
 		bool bIsOverrideActive = false;
 
 		// Registry of plugin-created transaction buffers.
-		TMap<FName, TObjectPtr<UTransBuffer>> BufferRegistry;
+		TMap<FString, TObjectPtr<UTransBuffer>> BufferRegistry;
 
 		// Reset all state to initial values.
 		void Reset()
@@ -42,7 +42,7 @@ namespace UE::AIAssistant
 
 static constexpr int32 DefaultUndoBufferSize = 256;
 
-TObjectPtr<UTransBuffer> FTransactionBufferManager::GetTransactionBuffer(const FName& BufferName)
+TObjectPtr<UTransBuffer> FTransactionBufferManager::GetTransactionBuffer(const FString& BufferName)
 {
 	check(IsInGameThread());
 
@@ -55,7 +55,7 @@ TObjectPtr<UTransBuffer> FTransactionBufferManager::GetTransactionBuffer(const F
 }
 
 TObjectPtr<UTransBuffer> FTransactionBufferManager::GetOrCreateTransactionBuffer(
-	const FName& BufferName)
+	const FString& BufferName)
 {
 	check(IsInGameThread());
 	check(GConfig);
@@ -82,7 +82,7 @@ TObjectPtr<UTransBuffer> FTransactionBufferManager::GetOrCreateTransactionBuffer
 	return NewBuffer;
 }
 
-bool FTransactionBufferManager::DestroyTransactionBuffer(const FName& BufferName)
+bool FTransactionBufferManager::DestroyTransactionBuffer(const FString& BufferName)
 {
 	check(IsInGameThread());
 	check(GEditor);
@@ -102,7 +102,7 @@ bool FTransactionBufferManager::DestroyTransactionBuffer(const FName& BufferName
 		UE_LOG(
 			LogAIAssistant, Warning,
 			TEXT("Cannot destroy transaction buffer '%s' while it is active as override"),
-			*BufferName.ToString());
+			*BufferName);
 		return false;
 	}
 
@@ -112,7 +112,7 @@ bool FTransactionBufferManager::DestroyTransactionBuffer(const FName& BufferName
 		UE_LOG(
 			LogAIAssistant, Warning,
 			TEXT("Cannot destroy transaction buffer '%s' while it is saved as the global buffer"),
-			*BufferName.ToString());
+			*BufferName);
 		return false;
 	}
 
@@ -210,6 +210,41 @@ bool FTransactionBufferManager::IsOverrideActive()
 	return TransactionBufferManagerState.bIsOverrideActive;
 }
 
+TSet<FString> FTransactionBufferManager::GetFilenamesFromUndoStack(const UTransBuffer* UndoStack)
+{
+	TSet<FString> Filenames;
+	if (!UndoStack)
+	{
+		return Filenames;
+	}
+
+	// Only look at transactions that haven't been undone (active transactions).
+	const int32 ActiveTransactionCount = UndoStack->GetQueueLength() - UndoStack->GetUndoCount();
+	for (int32 Index = 0; Index < ActiveTransactionCount; ++Index)
+	{
+		const FTransaction* Transaction = UndoStack->GetTransaction(Index);
+		if (Transaction)
+		{
+			// Get all objects that were part of this transaction.
+			TArray<UObject*> TransactionObjects;
+			Transaction->GetTransactionObjects(TransactionObjects);
+
+			for (UObject* Object : TransactionObjects)
+			{
+				// Get the package that contains this object.
+				UPackage* Package = Object ? Object->GetOutermost() : nullptr;
+				if (Package)
+				{
+					FString PackagePath = Package->GetPathName();
+					Filenames.Add(MoveTemp(PackagePath));
+				}
+			}
+		}
+	}
+
+	return Filenames;
+}
+
 }  // namespace UE::AIAssistant
 
 // Console commands for testing and debugging
@@ -225,18 +260,18 @@ void SetTransactionBufferOverrideCommand(const TArray<FString>& Args)
 		UE_LOG(LogAIAssistant, Warning, TEXT("Usage: TransactionBuffer.SetOverride <BufferName>"));
 		return;
 	}
-	FName BufferName = *Args[0];
+	const FString& BufferName = Args[0];
 	UTransBuffer* Buffer = FTransactionBufferManager::GetOrCreateTransactionBuffer(BufferName);
 	if (Buffer && FTransactionBufferManager::SetOverrideBuffer(Buffer))
 	{
 		UE_LOG(
 			LogAIAssistant, Log,
-			TEXT("Set transaction buffer override to '%s'"), *BufferName.ToString());
+			TEXT("Set transaction buffer override to '%s'"), *BufferName);
 	}
 	else
 	{
 		UE_LOG(LogAIAssistant, Error,
-			TEXT("Failed to set transaction buffer override to '%s'"), *BufferName.ToString());
+			TEXT("Failed to set transaction buffer override to '%s'"), *BufferName);
 	}
 }
 
@@ -256,18 +291,18 @@ void DestroyTransactionBufferCommand(const TArray<FString>& Args)
 		return;
 	}
 
-	FName BufferName(*Args[0]);
+	FString BufferName(Args[0]);
 	if (FTransactionBufferManager::DestroyTransactionBuffer(BufferName))
 	{
 		UE_LOG(
 			LogAIAssistant, Log,
-			TEXT("Destroyed transaction buffer '%s'"), *BufferName.ToString());
+			TEXT("Destroyed transaction buffer '%s'"), *BufferName);
 	}
 	else
 	{
 		UE_LOG(
 			LogAIAssistant, Error,
-			TEXT("Failed to destroy transaction buffer '%s'"), *BufferName.ToString());
+			TEXT("Failed to destroy transaction buffer '%s'"), *BufferName);
 	}
 }
 
diff --git a/Engine/Plugins/Experimental/AIAssistant/Source/AIAssistant/Private/AIAssistantTransactionBufferManager.h b/Engine/Plugins/Experimental/AIAssistant/Source/AIAssistant/Private/AIAssistantTransactionBufferManager.h
--- a/Engine/Plugins/Experimental/AIAssistant/Source/AIAssistant/Private/AIAssistantTransactionBufferManager.h
+++ b/Engine/Plugins/Experimental/AIAssistant/Source/AIAssistant/Private/AIAssistantTransactionBufferManager.h
@@ -17,14 +17,14 @@ public:
 	//
 	// @param BufferName Name of the transaction buffer to retrieve
 	// @return The transaction buffer if found, nullptr otherwise
-	static TObjectPtr<UTransBuffer> GetTransactionBuffer(const FName& BufferName);
+	static TObjectPtr<UTransBuffer> GetTransactionBuffer(const FString& BufferName);
 
 	// Create a new transaction buffer with the given name.
 	// If a buffer with this name already exists, returns the existing buffer.
 	//
 	// @param BufferName Unique name for the transaction buffer
 	// @return The created or existing transaction buffer
-	static TObjectPtr<UTransBuffer> GetOrCreateTransactionBuffer(const FName& BufferName);
+	static TObjectPtr<UTransBuffer> GetOrCreateTransactionBuffer(const FString& BufferName);
 	
 	// Destroy a transaction buffer by name.
 	// Removes it from the registry and releases the root reference.
@@ -32,7 +32,7 @@ public:
 	//
 	// @param BufferName Name of the transaction buffer to destroy
 	// @return True if the buffer was destroyed, false if not found or currently active
-	static bool DestroyTransactionBuffer(const FName& BufferName);
+	static bool DestroyTransactionBuffer(const FString& BufferName);
 
 	// Shutdown and cleanup all transaction buffers.
 	// Should be called during module shutdown.
@@ -54,6 +54,14 @@ public:
 	//
 	// @return True if an override buffer is active, false otherwise
 	static bool IsOverrideActive();
+
+	// Extract all package path names from an undo stack.
+	// Iterates through active transactions (those that haven't been undone) and collects
+	// the package paths of all objects that were modified.
+	//
+	// @param UndoStack The transaction buffer to extract filenames from
+	// @return Set of package path names for all modified objects
+	static TSet<FString> GetFilenamesFromUndoStack(const UTransBuffer* UndoStack);
 };
 
 }  // namespace UE::AIAssistant
diff --git a/Engine/Plugins/Experimental/AIAssistant/Source/AIAssistant/Private/Tests/AIAssistantTestObject.h b/Engine/Plugins/Experimental/AIAssistant/Source/AIAssistant/Private/Tests/AIAssistantTestObject.h
--- a/Engine/Plugins/Experimental/AIAssistant/Source/AIAssistant/Private/Tests/AIAssistantTestObject.h
+++ b/Engine/Plugins/Experimental/AIAssistant/Source/AIAssistant/Private/Tests/AIAssistantTestObject.h
@@ -24,6 +24,9 @@ public:
 		return Name;
 	}
 
+	UPROPERTY()
+	int32 TestValue = 0;
+
 private:
 	FString Name;
 };
\ No newline at end of file
diff --git a/Engine/Plugins/Experimental/AIAssistant/Source/AIAssistant/Private/Tests/AIAssistantTransactionBufferManagerTest.cpp b/Engine/Plugins/Experimental/AIAssistant/Source/AIAssistant/Private/Tests/AIAssistantTransactionBufferManagerTest.cpp
--- a/Engine/Plugins/Experimental/AIAssistant/Source/AIAssistant/Private/Tests/AIAssistantTransactionBufferManagerTest.cpp
+++ b/Engine/Plugins/Experimental/AIAssistant/Source/AIAssistant/Private/Tests/AIAssistantTransactionBufferManagerTest.cpp
@@ -7,6 +7,7 @@
 #include "Editor/TransBuffer.h"
 
 #include "AIAssistantTestFlags.h"
+#include "AIAssistantTestObject.h"
 #include "AIAssistantTransactionBufferManager.h"
 
 #if WITH_DEV_AUTOMATION_TESTS
@@ -16,8 +17,8 @@ using namespace UE::AIAssistant;
 BEGIN_DEFINE_SPEC(FAIAssistantTransactionBufferManagerTest,
 	"AI.Assistant.TransactionBufferManager",
 	AIAssistantTest::Flags)
-	FName TestBufferName;
-	FName TestBufferName2;
+	FString TestBufferName;
+	FString TestBufferName2;
 END_DEFINE_SPEC(FAIAssistantTransactionBufferManagerTest)
 
 void FAIAssistantTransactionBufferManagerTest::Define()
@@ -30,8 +31,8 @@ void FAIAssistantTransactionBufferManagerTest::Define()
 
 	BeforeEach([this]
 		{
-			TestBufferName = FName(TEXT("TestBuffer"));
-			TestBufferName2 = FName(TEXT("TestBuffer2"));
+			TestBufferName = TEXT("TestBuffer");
+			TestBufferName2 = TEXT("TestBuffer2");
 		});
 
 	AfterEach([this]
@@ -106,7 +107,7 @@ void FAIAssistantTransactionBufferManagerTest::Define()
 			It("returns false when buffer doesn't exist", [this]
 				{
 					const bool bResult = FTransactionBufferManager::DestroyTransactionBuffer(
-						FName(TEXT("NonExistentBuffer")));
+						TEXT("NonExistentBuffer"));
 					TestFalse(TEXT("Should return false for non-existent buffer"), bResult);
 				});
 
@@ -283,6 +284,82 @@ void FAIAssistantTransactionBufferManagerTest::Define()
 						FTransactionBufferManager::IsOverrideActive());
 				});
 		});
+
+	Describe("GetFilenamesFromUndoStack", [this]
+		{
+			It("returns empty set when UndoStack is nullptr", [this]
+				{
+					TSet<FString> Filenames =
+						FTransactionBufferManager::GetFilenamesFromUndoStack(nullptr);
+					TestEqual(
+						TEXT("Should return empty set for nullptr"),
+						Filenames.Num(), 0);
+				});
+
+			It("returns empty set when UndoStack has no transactions", [this]
+				{
+					TObjectPtr<UTransBuffer> Buffer =
+						FTransactionBufferManager::GetOrCreateTransactionBuffer(TestBufferName);
+					TSet<FString> Filenames =
+						FTransactionBufferManager::GetFilenamesFromUndoStack(Buffer);
+					TestEqual(
+						TEXT("Should return empty set for empty buffer"),
+						Filenames.Num(), 0);
+				});
+
+			It("returns filenames from active transactions", [this]
+				{
+					TObjectPtr<UTransBuffer> Buffer =
+						FTransactionBufferManager::GetOrCreateTransactionBuffer(TestBufferName);
+
+					// Create a test object and record a transaction
+					UPackage* TestPackage = NewObject<UPackage>(
+						nullptr, TEXT("/Temp/AIAssistantTestPackage"), RF_Transient);
+					UAIAssistantTestObject* TestObject = NewObject<UAIAssistantTestObject>(
+						TestPackage, TEXT("TestObject"), RF_Transactional);
+
+					// Begin a transaction on the buffer
+					Buffer->Begin(TEXT("Test Transaction"), FText::GetEmpty());
+					TestObject->Modify();
+					TestObject->TestValue = 1;
+					Buffer->End();
+
+					TSet<FString> Filenames =
+						FTransactionBufferManager::GetFilenamesFromUndoStack(Buffer);
+
+					TestTrue(
+						TEXT("Should contain the test package path"),
+						Filenames.Contains(TEXT("/Temp/AIAssistantTestPackage")));
+				});
+
+			It("does not return filenames from undone transactions", [this]
+				{
+					TObjectPtr<UTransBuffer> Buffer =
+						FTransactionBufferManager::GetOrCreateTransactionBuffer(TestBufferName);
+
+					// Create a test object and record a transaction
+					UPackage* TestPackage = NewObject<UPackage>(
+						nullptr, TEXT("/Temp/AIAssistantUndoTestPackage"), RF_Transient);
+					UAIAssistantTestObject* TestObject = NewObject<UAIAssistantTestObject>(
+						TestPackage, TEXT("TestObject"), RF_Transactional);
+
+					// Record a transaction
+					Buffer->Begin(TEXT("Test Transaction"), FText::GetEmpty());
+					TestObject->Modify();
+					TestObject->TestValue = 1;
+					Buffer->End();
+
+					// Undo the transaction
+					Buffer->Undo();
+
+					TSet<FString> Filenames =
+						FTransactionBufferManager::GetFilenamesFromUndoStack(Buffer);
+
+					TestFalse(
+						TEXT("Should not contain the undone package path"),
+						Filenames.Contains(TEXT("/Temp/AIAssistantUndoTestPackage")));
+				});
+		});
 }
 
 #endif  // WITH_DEV_AUTOMATION_TESTS
\ No newline at end of file
