Add pending file list WebApi functions.

#Synced-CL 50702752

diff --git a/Engine/Plugins/Experimental/AIAssistant/Source/AIAssistant/Private/AIAssistantWebApi.cpp b/Engine/Plugins/Experimental/AIAssistant/Source/AIAssistant/Private/AIAssistantWebApi.cpp
--- a/Engine/Plugins/Experimental/AIAssistant/Source/AIAssistant/Private/AIAssistantWebApi.cpp
+++ b/Engine/Plugins/Experimental/AIAssistant/Source/AIAssistant/Private/AIAssistantWebApi.cpp
@@ -16,6 +16,7 @@ namespace UE::AIAssistant
 {
 	UE_ENUM_METADATA_DEFINE(EMessageRole, UE_AI_ASSISTANT_MESSAGE_ROLE_ENUM);
 	UE_ENUM_METADATA_DEFINE(EMessageContentType, UE_AI_ASSISTANT_MESSAGE_CONTENT_TYPE_ENUM);
+	UE_ENUM_METADATA_DEFINE(EPendingFileStatus, UE_AI_ASSISTANT_PENDING_FILE_STATUS_ENUM);
 
 	const FString FWebApi::WebApiObjectName =
 		UE_AIASSISTANT_WEB_API_PARENT_OBJECT_NAME TEXT(".")
@@ -174,6 +175,63 @@ try {
 		WebJavaScriptResultDelegate->UnregisterResultHandler(HandlerId);
 	}
 
+	void FWebApi::UpdatePendingFileList(const FUpdatePendingFileListOptions& Options)
+	{
+		(void)ExecuteFunctionWithJsonArgument(
+			*WebApiObjectName, TEXT("updatePendingFileList"), Options);
+	}
+
+	TFuture<TValueOrError<FWebApi::FPendingFileDecisionHandle, FString>>
+		FWebApi::RegisterOnPendingFileDecision(
+		FOnPendingFileDecisionCallback OnPendingFileDecision)
+	{
+		using FResult = UAIAssistantWebJavaScriptResultDelegate::FResult;
+		FString HandlerId = RegisterCallback(
+			[OnPendingFileDecision = MoveTemp(OnPendingFileDecision)](FResult Result) {
+				FOnPendingFileDecisionOptions Options;
+				if (Result.bJsonIsError || !Options.FromJson(Result.Json))
+				{
+					UE_LOG(
+						LogAIAssistant, Error, TEXT("Javascript Error running callback: %s"),
+						*Result.Json);
+					return UAIAssistantWebJavaScriptResultDelegate::EHandlerStatus::Unregister;
+				}
+				OnPendingFileDecision(Options);
+				return UAIAssistantWebJavaScriptResultDelegate::EHandlerStatus::Continue;
+			});
+
+		TFuture<FResult> Future = ExecuteFunction(
+			*WebApiObjectName, TEXT("registerOnPendingFileDecision"),
+			*FormatFunctionCallback(HandlerId));
+
+		return Future.Next(
+			[HandlerId](const FResult& Result) -> TValueOrError<FPendingFileDecisionHandle, FString>
+			{
+				if (Result.bJsonIsError)
+				{
+					return MakeError(Result.Json);
+				}
+				return MakeValue(FPendingFileDecisionHandle
+					{
+						HandlerId,
+						Result.Json
+					});
+			});
+	}
+
+	void FWebApi::UnregisterOnPendingFileDecision(const FPendingFileDecisionHandle& Handle)
+	{
+		auto& [HandlerId, CallbackHandle] = Handle;
+		if (!WebJavaScriptResultDelegate.IsValid())
+		{
+			return;
+		}
+		(void)ExecuteFunction(
+			*WebApiObjectName, TEXT("unregisterOnPendingFileDecision"),
+			*FString::Printf(TEXT("%s"), *CallbackHandle));
+		WebJavaScriptResultDelegate->UnregisterResultHandler(HandlerId);
+	}
+
 	FString FWebApi::FormatFunctionCall(
 		const TCHAR* InstanceName, const TCHAR* FunctionName, const TCHAR* Arguments,
 		const FString& HandlerId) const
diff --git a/Engine/Plugins/Experimental/AIAssistant/Source/AIAssistant/Private/AIAssistantWebApi.h b/Engine/Plugins/Experimental/AIAssistant/Source/AIAssistant/Private/AIAssistantWebApi.h
--- a/Engine/Plugins/Experimental/AIAssistant/Source/AIAssistant/Private/AIAssistantWebApi.h
+++ b/Engine/Plugins/Experimental/AIAssistant/Source/AIAssistant/Private/AIAssistantWebApi.h
@@ -258,6 +258,71 @@ namespace UE::AIAssistant
 		END_JSON_SERIALIZER
 	};
 
+	// Status indicating how a file was changed.
+	enum class EPendingFileStatus
+	{
+		Added,
+		Removed,
+		Modified,
+		Moved,
+	};
+
+	#define UE_AI_ASSISTANT_PENDING_FILE_STATUS_ENUM(X) \
+		X(EPendingFileStatus::Added, "added"), \
+		X(EPendingFileStatus::Removed, "removed"), \
+		X(EPendingFileStatus::Modified, "modified"), \
+		X(EPendingFileStatus::Moved, "moved")
+
+	UE_ENUM_METADATA_DECLARE(EPendingFileStatus, UE_AI_ASSISTANT_PENDING_FILE_STATUS_ENUM);
+
+	// Metadata describing a file that has been modified by tool calls.
+	struct FPendingFileMetadata : public FJsonSerializable
+	{
+		// Display name shown to the user (e.g., filename without path).
+		FString DisplayName;
+		// Full path to the source file.
+		FString FullPath;
+		// Status indicating how the file was changed.
+		EPendingFileStatus Status = EPendingFileStatus::Modified;
+		// For moved files, the original path before the move.
+		TOptional<FString> OriginalPath;
+
+		BEGIN_JSON_SERIALIZER
+			JSON_SERIALIZE("displayName", DisplayName);
+			JSON_SERIALIZE("fullPath", FullPath);
+			JSON_SERIALIZE_ENUM("status", Status);
+			JSON_SERIALIZE_OPTIONAL("originalPath", OriginalPath);
+		END_JSON_SERIALIZER
+	};
+
+	// Argument for UpdatePendingFileList.
+	struct FUpdatePendingFileListOptions : public FJsonSerializable
+	{
+		// ID of the conversation this file list belongs to.
+		// If unspecified, the currently active conversation is used.
+		TOptional<FConversationId> ConversationId;
+		// List of file metadata describing the pending file changes.
+		TArray<FPendingFileMetadata> Files;
+
+		BEGIN_JSON_SERIALIZER
+			JSON_SERIALIZE_OPTIONAL_OBJECT_SERIALIZABLE("conversationId", ConversationId);
+			JSON_SERIALIZE_ARRAY_SERIALIZABLE("files", Files, FPendingFileMetadata);
+		END_JSON_SERIALIZER
+	};
+	
+	// Argument for OnPendingFileDecision callback.
+	struct FOnPendingFileDecisionOptions : public FJsonSerializable
+	{
+		// Unique ID of the conversation.
+		FConversationId ConversationId;
+		bool bAccepted = false;
+
+		BEGIN_JSON_SERIALIZER
+			JSON_SERIALIZE_OBJECT_SERIALIZABLE("conversationId", ConversationId);
+			JSON_SERIALIZE("accepted", bAccepted);
+		END_JSON_SERIALIZER
+	};
+
 	// High level descriptor of the environment is interacting with.
 	struct FAgentEnvironmentDescriptor : public FJsonSerializable {
 		// Name of the current environment.
@@ -431,6 +496,24 @@ namespace UE::AIAssistant
 		// RegisterOnConversationUpdate.
 		void UnregisterOnConversationUpdate(const FConversationUpdateHandle& Registration);
 
+		// Update the list of pending files modified by tool calls for a conversation.
+		// Called by the C++ backend after each tool call to report which files were modified.
+		void UpdatePendingFileList(const FUpdatePendingFileListOptions& Options);
+
+		// Opaque handle to data needed to unregister pending file decision callback.
+		using FPendingFileDecisionHandle = TTuple<FString, FString>;
+
+		// Register for pending file decision events.
+		using FOnPendingFileDecisionCallback = TFunction<void(FOnPendingFileDecisionOptions)>;
+
+		// Register a callback for when the user approves or rejects pending file changes.
+		TFuture<TValueOrError<FPendingFileDecisionHandle, FString>> RegisterOnPendingFileDecision(
+			FOnPendingFileDecisionCallback OnPendingFileDecision);
+
+		// Unregister pending file decision callback using the handle returned from
+		// RegisterOnPendingFileDecision.
+		void UnregisterOnPendingFileDecision(const FPendingFileDecisionHandle& Registration);
+
 	protected:
 		// Format a function call of a member with result handling.
 		FString FormatFunctionCall(
diff --git a/Engine/Plugins/Experimental/AIAssistant/Source/AIAssistant/Private/Tests/AIAssistantWebApiTest.cpp b/Engine/Plugins/Experimental/AIAssistant/Source/AIAssistant/Private/Tests/AIAssistantWebApiTest.cpp
--- a/Engine/Plugins/Experimental/AIAssistant/Source/AIAssistant/Private/Tests/AIAssistantWebApiTest.cpp
+++ b/Engine/Plugins/Experimental/AIAssistant/Source/AIAssistant/Private/Tests/AIAssistantWebApiTest.cpp
@@ -586,6 +586,169 @@ void FAIAssistantWebApiTest::Define()
 				});
 		}
 	);
+
+	Describe("UpdatePendingFileList", [this]
+	{
+		It("calls updatePendingFileList with options", [this]
+		{
+			FFakeWebApiContainer WebApi;
+			FUpdatePendingFileListOptions Options;
+			Options.ConversationId.Emplace();
+			Options.ConversationId->Id = TEXT("conv-1");
+			Options.ConversationId->Type = TEXT("ConversationId");
+
+			FPendingFileMetadata& File = Options.Files.AddDefaulted_GetRef();
+			File.DisplayName = TEXT("main.cpp");
+			File.FullPath = TEXT("/src/main.cpp");
+			File.Status = EPendingFileStatus::Modified;
+
+			WebApi->UpdatePendingFileList(Options);
+			TestTrue(TEXT("Expected updatePendingFileList call"),
+				WebApi->TestExpectAsyncFunctionCall(
+					*this, nullptr, TEXT("updatePendingFileList"), *Options.ToJson(false)).IsValid());
+		});
+	});
+
+	Describe("RegisterOnPendingFileDecision", [this]
+	{
+		It("registers callback without invoking", [this]
+		{
+			FFakeWebApiContainer WebApi;
+			int CallCount = 0;
+			auto Future =
+				WebApi->RegisterOnPendingFileDecision(
+					[&CallCount](FOnPendingFileDecisionOptions Options)
+					{
+						CallCount += 1;
+					});
+			const FString& HandlerId = WebApi->RegisteredCallbackHandlerIds.Last();
+			WebApi->TestExpectAsyncFunctionCallAndComplete(
+				*this,
+				nullptr,
+				TEXT("registerOnPendingFileDecision"),
+				*FWebApiAccessor::FormatFunctionCallback(*WebApi, *HandlerId),
+				TEXT(""),
+				false);
+			TestEqual(TEXT("RegisterOnPendingFileDecision_NoCalls"), CallCount, 0);
+		});
+
+		It("invokes callback with accepted decision", [this]
+		{
+			FFakeWebApiContainer WebApi;
+			FString ActualConversationId;
+			bool bActualAccepted = false;
+			int CallCount = 0;
+			auto Result = WebApi->RegisterOnPendingFileDecision(
+				[&ActualConversationId, &bActualAccepted, &CallCount](FOnPendingFileDecisionOptions Options)
+				{
+					ActualConversationId = Options.ConversationId.Id;
+					bActualAccepted = Options.bAccepted;
+					CallCount += 1;
+				});
+			const FString& HandlerId = WebApi->RegisteredCallbackHandlerIds.Last();
+			if (!WebApi->TestExpectAsyncFunctionCallAndComplete(
+				*this,
+				nullptr,
+				TEXT("registerOnPendingFileDecision"),
+				*FWebApiAccessor::FormatFunctionCallback(*WebApi, *HandlerId),
+				TEXT(""),
+				false))
+			{
+				return;
+			}
+
+			FString ExpectedConversationId = TEXT("conv-123");
+			WebApi->InvokeCallback(
+				HandlerId,
+				TEXT(R"js({ "conversationId": { "id": "conv-123", "type": "ConversationId" }, "accepted": true })js"),
+				false);
+			TestEqual(
+				TEXT("RegisterOnPendingFileDecision_ConversationId"),
+				ActualConversationId,
+				ExpectedConversationId);
+			TestTrue(TEXT("RegisterOnPendingFileDecision_Accepted"), bActualAccepted);
+			TestEqual(TEXT("RegisterOnPendingFileDecision_CallCount"), CallCount, 1);
+		});
+
+		It("invokes callback with rejected decision", [this]
+		{
+			FFakeWebApiContainer WebApi;
+			bool bActualAccepted = true;
+			auto Result = WebApi->RegisterOnPendingFileDecision(
+				[&bActualAccepted](FOnPendingFileDecisionOptions Options)
+				{
+					bActualAccepted = Options.bAccepted;
+				});
+			const FString& HandlerId = WebApi->RegisteredCallbackHandlerIds.Last();
+			if (!WebApi->TestExpectAsyncFunctionCallAndComplete(
+				*this,
+				nullptr,
+				TEXT("registerOnPendingFileDecision"),
+				*FWebApiAccessor::FormatFunctionCallback(*WebApi, *HandlerId),
+				TEXT(""),
+				false))
+			{
+				return;
+			}
+
+			WebApi->InvokeCallback(
+				HandlerId,
+				TEXT(R"js({ "conversationId": { "id": "conv-1", "type": "ConversationId" }, "accepted": false })js"),
+				false);
+			TestFalse(TEXT("RegisterOnPendingFileDecision_Rejected"), bActualAccepted);
+		});
+
+		It("unregisters via UnregisterOnPendingFileDecision", [this]
+		{
+			FFakeWebApiContainer WebApi;
+
+			// Register a callback.
+			int CallCount = 0;
+			auto Result = WebApi->RegisterOnPendingFileDecision(
+				[&CallCount](FOnPendingFileDecisionOptions Options)
+				{
+					CallCount += 1;
+				});
+			const FString& HandlerId = WebApi->RegisteredCallbackHandlerIds.Last();
+			if (!WebApi->TestExpectAsyncFunctionCallAndComplete(
+				*this,
+				nullptr,
+				TEXT("registerOnPendingFileDecision"),
+				*FWebApiAccessor::FormatFunctionCallback(*WebApi, *HandlerId),
+				TEXT(""),
+				false))
+			{
+				return;
+			}
+
+			auto& Registration = Result.Get().GetValue();
+			auto& CallbackHandle = Registration.Get<1>();
+
+			// Invoke callback once.
+			WebApi->InvokeCallback(
+				HandlerId,
+				TEXT(R"js({ "conversationId": { "id": "conv-1", "type": "ConversationId" }, "accepted": true })js"),
+				false);
+
+			// Unregister the callback.
+			WebApi->UnregisterOnPendingFileDecision(Registration);
+			if (!WebApi->TestExpectAsyncFunctionCall(
+				*this, nullptr, TEXT("unregisterOnPendingFileDecision"),
+				*FString::Printf(TEXT("%s"), *CallbackHandle)).IsValid())
+			{
+				return;
+			}
+
+			// Try to invoke callback again.
+			WebApi->InvokeCallback(
+				HandlerId,
+				TEXT(R"js({ "conversationId": { "id": "conv-2", "type": "ConversationId" }, "accepted": false })js"),
+				false);
+
+			// Ensure callback was only called the first time.
+			TestEqual(TEXT("UnregisterOnPendingFileDecision_CallCount"), CallCount, 1);
+		});
+	});
 }
 
 
@@ -915,4 +1078,166 @@ bool FAIAssistantWebApiTestUnregisterOnConversationUpdate::RunTest(const FString
 		&& TestEqual(TEXT("UnregisterOnConversationUpdate_CallCount"), CallCount, 1);
 }
 
+BEGIN_DEFINE_SPEC(AIAssistantWebApiPendingFileMetadata, "AI.Assistant.WebApi.PendingFileMetadata",
+	AIAssistantTest::Flags)
+END_DEFINE_SPEC(AIAssistantWebApiPendingFileMetadata)
+
+void AIAssistantWebApiPendingFileMetadata::Define()
+{
+	Describe("JSON Serialization", [this]
+	{
+		It("serializes all fields correctly", [this]
+		{
+			FPendingFileMetadata Metadata;
+			Metadata.DisplayName = TEXT("main.cpp");
+			Metadata.FullPath = TEXT("/src/main.cpp");
+			Metadata.Status = EPendingFileStatus::Modified;
+
+			const FString Json = Metadata.ToJson(false);
+			TSharedPtr<FJsonObject> Object;
+			const TSharedRef<TJsonReader<>> Reader = TJsonReaderFactory<>::Create(Json);
+			if (!TestTrue(TEXT("Should parse as valid JSON"),
+				FJsonSerializer::Deserialize(Reader, Object) && Object.IsValid()))
+			{
+				return;
+			}
+
+			FString DisplayName, FullPath, Status;
+			if (!TestTrue(TEXT("displayName exists"), Object->TryGetStringField(TEXT("displayName"), DisplayName)) ||
+				!TestTrue(TEXT("fullPath exists"), Object->TryGetStringField(TEXT("fullPath"), FullPath)) ||
+				!TestTrue(TEXT("status exists"), Object->TryGetStringField(TEXT("status"), Status)))
+			{
+				return;
+			}
+			TestEqual(TEXT("displayName"), DisplayName, TEXT("main.cpp"));
+			TestEqual(TEXT("fullPath"), FullPath, TEXT("/src/main.cpp"));
+			TestEqual(TEXT("status"), Status, TEXT("modified"));
+		});
+
+		It("serializes optional originalPath when set", [this]
+		{
+			FPendingFileMetadata Metadata;
+			Metadata.DisplayName = TEXT("config.json");
+			Metadata.FullPath = TEXT("/new/config.json");
+			Metadata.Status = EPendingFileStatus::Moved;
+			Metadata.OriginalPath = TEXT("/old/config.json");
+
+			const FString Json = Metadata.ToJson(false);
+			TSharedPtr<FJsonObject> Object;
+			const TSharedRef<TJsonReader<>> Reader = TJsonReaderFactory<>::Create(Json);
+			if (!TestTrue(TEXT("Should parse as valid JSON"),
+				FJsonSerializer::Deserialize(Reader, Object) && Object.IsValid()))
+			{
+				return;
+			}
+
+			TestEqual(TEXT("status"), Object->GetStringField(TEXT("status")), TEXT("moved"));
+			FString OriginalPath;
+			if (!TestTrue(TEXT("originalPath exists"), Object->TryGetStringField(TEXT("originalPath"), OriginalPath)))
+			{
+				return;
+			}
+			TestEqual(TEXT("originalPath"), OriginalPath, TEXT("/old/config.json"));
+		});
+
+		It("parses all status values correctly", [this]
+		{
+			auto TestStatus = [this](const TCHAR* StatusStr, EPendingFileStatus ExpectedStatus)
+			{
+				const FString Input = FString::Printf(
+					TEXT(R"json({"displayName":"test","fullPath":"/test","status":"%s"})json"),
+					StatusStr);
+				FPendingFileMetadata Metadata;
+				if (!TestTrue(FString::Printf(TEXT("FromJson should succeed for %s"), StatusStr),
+					Metadata.FromJson(Input)))
+				{
+					return;
+				}
+				TestEqual(FString::Printf(TEXT("Status should be %s"), StatusStr),
+					Metadata.Status, ExpectedStatus);
+			};
+
+			TestStatus(TEXT("added"), EPendingFileStatus::Added);
+			TestStatus(TEXT("removed"), EPendingFileStatus::Removed);
+			TestStatus(TEXT("modified"), EPendingFileStatus::Modified);
+			TestStatus(TEXT("moved"), EPendingFileStatus::Moved);
+		});
+	});
+}
+
+BEGIN_DEFINE_SPEC(AIAssistantWebApiUpdatePendingFileListOptions,
+	"AI.Assistant.WebApi.UpdatePendingFileListOptions", AIAssistantTest::Flags)
+END_DEFINE_SPEC(AIAssistantWebApiUpdatePendingFileListOptions)
+
+void AIAssistantWebApiUpdatePendingFileListOptions::Define()
+{
+	Describe("JSON Serialization", [this]
+	{
+		It("serializes with conversationId when set", [this]
+		{
+			FUpdatePendingFileListOptions Options;
+			Options.ConversationId.Emplace();
+			Options.ConversationId->Id = TEXT("conv-123");
+			Options.ConversationId->Type = TEXT("ConversationId");
+
+			FPendingFileMetadata& File = Options.Files.AddDefaulted_GetRef();
+			File.DisplayName = TEXT("test.cpp");
+			File.FullPath = TEXT("/src/test.cpp");
+			File.Status = EPendingFileStatus::Added;
+
+			const FString Json = Options.ToJson(false);
+			TSharedPtr<FJsonObject> Object;
+			const TSharedRef<TJsonReader<>> Reader = TJsonReaderFactory<>::Create(Json);
+			if (!TestTrue(TEXT("Should parse as valid JSON"),
+				FJsonSerializer::Deserialize(Reader, Object) && Object.IsValid()))
+			{
+				return;
+			}
+
+			const TSharedPtr<FJsonObject>* ConvIdObj;
+			if (!TestTrue(TEXT("conversationId exists"), Object->TryGetObjectField(TEXT("conversationId"), ConvIdObj)))
+			{
+				return;
+			}
+			TestEqual(TEXT("conversationId.id"), (*ConvIdObj)->GetStringField(TEXT("id")), TEXT("conv-123"));
+
+			const TArray<TSharedPtr<FJsonValue>>* FilesArray;
+			if (!TestTrue(TEXT("files array exists"), Object->TryGetArrayField(TEXT("files"), FilesArray)))
+			{
+				return;
+			}
+			TestEqual(TEXT("files count"), FilesArray->Num(), 1);
+		});
+
+		It("serializes without conversationId when not set", [this]
+		{
+			FUpdatePendingFileListOptions Options;
+			// ConversationId not set
+
+			FPendingFileMetadata& File = Options.Files.AddDefaulted_GetRef();
+			File.DisplayName = TEXT("test.cpp");
+			File.FullPath = TEXT("/src/test.cpp");
+			File.Status = EPendingFileStatus::Removed;
+
+			const FString Json = Options.ToJson(false);
+			TSharedPtr<FJsonObject> Object;
+			const TSharedRef<TJsonReader<>> Reader = TJsonReaderFactory<>::Create(Json);
+			if (!TestTrue(TEXT("Should parse as valid JSON"),
+				FJsonSerializer::Deserialize(Reader, Object) && Object.IsValid()))
+			{
+				return;
+			}
+
+			TestFalse(TEXT("conversationId should not exist"), Object->HasField(TEXT("conversationId")));
+
+			const TArray<TSharedPtr<FJsonValue>>* FilesArray;
+			if (!TestTrue(TEXT("files array exists"), Object->TryGetArrayField(TEXT("files"), FilesArray)))
+			{
+				return;
+			}
+			TestEqual(TEXT("files count"), FilesArray->Num(), 1);
+		});
+	});
+}
+
 #endif  // WITH_DEV_AUTOMATION_TESTS
