Adds support for per-project and per-user global context in the EDA as well as a tool call to access them.

#Synced-CL 50691082

diff --git a/Engine/Plugins/Experimental/AIAssistant/Source/AIAssistant/Private/AIAssistantToolset.cpp b/Engine/Plugins/Experimental/AIAssistant/Source/AIAssistant/Private/AIAssistantToolset.cpp
--- a/Engine/Plugins/Experimental/AIAssistant/Source/AIAssistant/Private/AIAssistantToolset.cpp
+++ b/Engine/Plugins/Experimental/AIAssistant/Source/AIAssistant/Private/AIAssistantToolset.cpp
@@ -22,6 +22,26 @@ using namespace UE::AIAssistant::SlateQuerier;
 #include UE_INLINE_GENERATED_CPP_BY_NAME(AIAssistantToolset)
 
 
+FString UAIAssistantToolset::GetProjectContext()
+{
+	FString CombinedContext;
+	const UAIAssistantContextProject* ProjectSettings = GetDefault<UAIAssistantContextProject>();
+	if (ProjectSettings && !ProjectSettings->Prompt.IsEmpty())
+	{
+		CombinedContext += ProjectSettings->Prompt;		
+	}
+	const UAIAssistantContextUser* UserSettings = GetDefault<UAIAssistantContextUser>();
+	if (UserSettings && !UserSettings->Prompt.IsEmpty())
+	{
+		if (!CombinedContext.IsEmpty())
+		{			
+			CombinedContext += FString(TEXT("\n"));
+		}
+		CombinedContext += UserSettings->Prompt;		
+	}
+	return CombinedContext;
+}
+
 FAIAssistantDockContext UAIAssistantToolset::GetDockedContext()
 {	
 	TSharedPtr<SAIAssistantWebBrowser> BrowserWidget = UAIAssistantSubsystem::GetAIAssistantWebBrowserWidget();
diff --git a/Engine/Plugins/Experimental/AIAssistant/Source/AIAssistant/Private/AIAssistantToolset.h b/Engine/Plugins/Experimental/AIAssistant/Source/AIAssistant/Private/AIAssistantToolset.h
--- a/Engine/Plugins/Experimental/AIAssistant/Source/AIAssistant/Private/AIAssistantToolset.h
+++ b/Engine/Plugins/Experimental/AIAssistant/Source/AIAssistant/Private/AIAssistantToolset.h
@@ -2,13 +2,39 @@
 
 #pragma once
 
+#include "Engine/DeveloperSettings.h"
 #include "ToolsetRegistry/Public/BPToolsetRegistry.h"
 
+
 #include "AIAssistantToolset.generated.h"
 
 class UEdGraph;
 class UEdGraphNode;
 
+UCLASS(config=EditorPerProjectUserSettings, meta=(DisplayName="Epic Developer Assistant"))
+class UAIAssistantContextUser : public UDeveloperSettings
+{
+	GENERATED_BODY()
+public:
+	/// This field is used to help the Epic Developer Assistant understand important info about you,
+	/// and can include things like your role, preferences, current goals, etc. The EDA
+	/// will automatically consider this information whenever you ask it to do anything.
+	UPROPERTY(Config, EditAnywhere, Category="Epic Developer Assistant", Meta=(Multiline=true))
+	FString Prompt;
+};
+
+UCLASS(config=Editor, defaultconfig, meta=(DisplayName="Epic Developer Assistant"))
+class UAIAssistantContextProject : public UDeveloperSettings
+{
+	GENERATED_BODY()
+public:
+	/// This field is used to help the Epic Developer Assistant understand important info about your
+	/// project, and can include things like game genre, naming conventions, art style, etc. The EDA
+	/// will automatically consider this information whenever you ask it to do anything.
+	UPROPERTY(Config, EditAnywhere, Category = "Epic Developer Assistant", Meta=(Multiline=true))
+	FString Prompt;
+};
+
 /// Contains information about the state of the editor that the assistant is currently docked with.
 USTRUCT(BlueprintType)
 struct FAIAssistantDockContext
@@ -35,6 +61,10 @@ class UAIAssistantToolset : public UToolsetDefinition
 {
 	GENERATED_BODY()
 public:
+	/// Returns important information about the current project and/or current user.
+	UFUNCTION(BlueprintCallable, Category = "AI Assistant")
+	static FString GetProjectContext();
+
 	/// Returns information about the part of the unreal UI that the assistant is docked into.	
 	/// For example, if the assistant is docked with a Blueprint asset editor, this return
 	/// the Blueprint asset that is being edited, which graph is in focus, etc.
diff --git a/Engine/Plugins/Experimental/AIAssistant/Source/AIAssistant/Private/Tests/AIAssistantToolsetTest.cpp b/Engine/Plugins/Experimental/AIAssistant/Source/AIAssistant/Private/Tests/AIAssistantToolsetTest.cpp
new file mode 100644
--- /dev/null
+++ b/Engine/Plugins/Experimental/AIAssistant/Source/AIAssistant/Private/Tests/AIAssistantToolsetTest.cpp
@@ -0,0 +1,80 @@
+// Copyright Epic Games, Inc. All Rights Reserved.
+
+#include "Misc/AutomationTest.h"
+
+#include "AIAssistantToolset.h"
+#include "AIAssistantTestFlags.h"
+
+#if WITH_DEV_AUTOMATION_TESTS
+
+
+BEGIN_DEFINE_SPEC(FAssistantToolsetSpec, "AI.Assistant.Toolset", AIAssistantTest::Flags)
+	FString ProjectPrompt;
+	FString UserPrompt;	
+END_DEFINE_SPEC(FAssistantToolsetSpec)
+
+void FAssistantToolsetSpec::Define()
+{
+	BeforeEach([this]()
+	{
+		const UAIAssistantContextProject* ProjectSettings = GetDefault<UAIAssistantContextProject>();
+		if (ProjectSettings)
+		{
+			ProjectPrompt = ProjectSettings->Prompt;
+		}
+		const UAIAssistantContextUser* UserSettings = GetDefault<UAIAssistantContextUser>();
+		if (UserSettings)
+		{
+			UserPrompt = UserSettings->Prompt;
+		}
+	});
+
+	AfterEach([this]()
+	{
+		UAIAssistantContextProject* ProjectSettings = GetMutableDefault<UAIAssistantContextProject>();
+		if (ProjectSettings)
+		{
+			ProjectSettings->Prompt = ProjectPrompt;
+		}
+		ProjectPrompt.Empty();
+		UAIAssistantContextUser* UserSettings = GetMutableDefault<UAIAssistantContextUser>();
+		if (UserSettings)
+		{
+			UserSettings->Prompt = UserPrompt;
+		}
+		UserPrompt.Empty();
+	});
+
+	It("Returns project context.", [this]()
+	{
+		UAIAssistantContextProject* ProjectSettings = GetMutableDefault<UAIAssistantContextProject>();
+		ProjectSettings->Prompt = FString(TEXT("Foo bar"));
+
+		FString Context = UAIAssistantToolset::GetProjectContext();
+		TestEqual(TEXT("Context"), Context, FString(TEXT("Foo bar")));
+	});
+
+	It("Returns user context.", [this]()
+	{
+		UAIAssistantContextUser* UserSettings = GetMutableDefault<UAIAssistantContextUser>();
+		UserSettings->Prompt = FString(TEXT("Bish bosh"));
+
+		FString Context = UAIAssistantToolset::GetProjectContext();
+		TestEqual(TEXT("Context"), Context, FString(TEXT("Bish bosh")));
+	});
+
+	It("Returns project and user context.", [this]()
+	{
+		UAIAssistantContextProject* ProjectSettings = GetMutableDefault<UAIAssistantContextProject>();
+		ProjectSettings->Prompt = FString(TEXT("Foo bar"));
+
+		UAIAssistantContextUser* UserSettings = GetMutableDefault<UAIAssistantContextUser>();
+		UserSettings->Prompt = FString(TEXT("Bish bosh"));
+
+		FString Context = UAIAssistantToolset::GetProjectContext();
+		TestEqual(TEXT("Context"), Context, FString(TEXT("Foo bar\nBish bosh")));
+	});
+}
+
+
+#endif  // WITH_DEV_AUTOMATION_TESTS
