[AI Assistant] Fix GetDockedContext tool when it is invoked through mechanism that doesn't use the widget to interact with toolsets, e.g. a MCP server, and thus gets a null pointer.

#rnx

#Synced-CL 50868869

diff --git a/Engine/Plugins/Experimental/AIAssistant/Source/AIAssistant/Private/AIAssistantToolset.cpp b/Engine/Plugins/Experimental/AIAssistant/Source/AIAssistant/Private/AIAssistantToolset.cpp
--- a/Engine/Plugins/Experimental/AIAssistant/Source/AIAssistant/Private/AIAssistantToolset.cpp
+++ b/Engine/Plugins/Experimental/AIAssistant/Source/AIAssistant/Private/AIAssistantToolset.cpp
@@ -12,6 +12,7 @@
 #include "UObject/ObjectPtr.h"
 #include "Widgets/Docking/SDockTab.h"
 
+#include "AIAssistant.h"
 #include "AIAssistantSubsystem.h"
 #include "AIAssistantSlateQuerierUtils.h"
 #include "AIAssistantSlateQuerierUtilsTypeSearch.h"
@@ -43,8 +44,13 @@ FString UAIAssistantToolset::GetProjectContext()
 }
 
 FAIAssistantDockContext UAIAssistantToolset::GetDockedContext()
-{	
-	TSharedPtr<SAIAssistantWebBrowser> BrowserWidget = UAIAssistantSubsystem::GetAIAssistantWebBrowserWidget();
+{
+	TSharedPtr<SAIAssistantWebBrowser> BrowserWidget =
+		UAIAssistantSubsystem::GetAIAssistantModule().GetAIAssistantWebBrowserWidget();
+	if (!BrowserWidget)
+	{
+		return FAIAssistantDockContext();
+	}
 
 	// Walk up the widget hierarchy until we encounter an asset editor.
 	TSharedPtr<SWidget> Current = BrowserWidget;
diff --git a/Engine/Plugins/Experimental/AIAssistant/Source/AIAssistant/Private/AIAssistantToolset.h b/Engine/Plugins/Experimental/AIAssistant/Source/AIAssistant/Private/AIAssistantToolset.h
--- a/Engine/Plugins/Experimental/AIAssistant/Source/AIAssistant/Private/AIAssistantToolset.h
+++ b/Engine/Plugins/Experimental/AIAssistant/Source/AIAssistant/Private/AIAssistantToolset.h
@@ -65,9 +65,10 @@ public:
 	UFUNCTION(BlueprintCallable, Category = "AI Assistant")
 	static FString GetProjectContext();
 
-	/// Returns information about the part of the unreal UI that the assistant is docked into.	
-	/// For example, if the assistant is docked with a Blueprint asset editor, this return
-	/// the Blueprint asset that is being edited, which graph is in focus, etc.
+	/// Returns information about the part of the unreal UI that the assistant is docked into.
+	/// For example, if the assistant is docked with a Blueprint asset editor, this returns
+	/// the Blueprint asset that is being edited, which graph is in focus, etc.  If the assistant
+	/// is closed or not docked to an asset editor, this information will be empty.
 	UFUNCTION(BlueprintCallable, Category = "AI Assistant")
 	static FAIAssistantDockContext GetDockedContext();
 
