Make assistant project context return structured data that clearly delineates user vs project context. Also add engine context which will replace the intro toolset.

#Synced-CL 50942406

diff --git a/Engine/Plugins/Experimental/AIAssistant/Source/AIAssistant/Private/AIAssistantToolset.cpp b/Engine/Plugins/Experimental/AIAssistant/Source/AIAssistant/Private/AIAssistantToolset.cpp
--- a/Engine/Plugins/Experimental/AIAssistant/Source/AIAssistant/Private/AIAssistantToolset.cpp
+++ b/Engine/Plugins/Experimental/AIAssistant/Source/AIAssistant/Private/AIAssistantToolset.cpp
@@ -22,22 +22,30 @@ using namespace UE::AIAssistant::SlateQuerier;
 #include UE_INLINE_GENERATED_CPP_BY_NAME(AIAssistantToolset)
 
 
-FString UAIAssistantToolset::GetProjectContext()
+FAIAssistantContext UAIAssistantToolset::GetProjectContext()
 {
-	FString CombinedContext;
+	FAIAssistantContext CombinedContext;
+
+	TArray<FString> EngineTips({
+		TEXT("Folder paths always need to end in /"),
+		TEXT("100 units is 1 meter"),
+		TEXT("Screen space coordinates are normalized coordinates where 0, 0 is the upper left corner "
+			 "and 1, 1 is the lower right corner."),
+		TEXT("Colors are [0,255]. LinearColors are [0,1]"),
+		TEXT("If you aren't sure what the user is referring to, try checked to see what they have selected.")
+	});
+
+	CombinedContext.UnrealContext = FString::Join(EngineTips, TEXT("\n"));
+
 	const UAIAssistantContextProject* ProjectSettings = GetDefault<UAIAssistantContextProject>();
-	if (ProjectSettings && !ProjectSettings->Prompt.IsEmpty())
+	if (ProjectSettings)
 	{
-		CombinedContext += ProjectSettings->Prompt;		
+		CombinedContext.ProjectContext = ProjectSettings->Prompt;		
 	}
 	const UAIAssistantContextUser* UserSettings = GetDefault<UAIAssistantContextUser>();
-	if (UserSettings && !UserSettings->Prompt.IsEmpty())
+	if (UserSettings)
 	{
-		if (!CombinedContext.IsEmpty())
-		{			
-			CombinedContext += FString(TEXT("\n"));
-		}
-		CombinedContext += UserSettings->Prompt;		
+		CombinedContext.UserContext = UserSettings->Prompt;		
 	}
 	return CombinedContext;
 }
diff --git a/Engine/Plugins/Experimental/AIAssistant/Source/AIAssistant/Private/AIAssistantToolset.h b/Engine/Plugins/Experimental/AIAssistant/Source/AIAssistant/Private/AIAssistantToolset.h
--- a/Engine/Plugins/Experimental/AIAssistant/Source/AIAssistant/Private/AIAssistantToolset.h
+++ b/Engine/Plugins/Experimental/AIAssistant/Source/AIAssistant/Private/AIAssistantToolset.h
@@ -35,12 +35,30 @@ public:
 	FString Prompt;
 };
 
+/// Contains information about the current Unreal project and user.
+USTRUCT(BlueprintType)
+struct FAIAssistantContext
+{
+	GENERATED_USTRUCT_BODY()
+public:
+	/// Describes important context about how to work with the Unreal engine.
+	UPROPERTY(BlueprintReadOnly, Category = "AIAssistant")
+	FString UnrealContext;
+
+	/// Describes important context about the currently active Unreal project.
+	UPROPERTY(BlueprintReadOnly, Category = "AIAssistant")
+	FString ProjectContext;
+
+	/// Describes important context about the currently active UnrealEd user.
+	UPROPERTY(BlueprintReadOnly, Category = "AIAssistant")
+	FString UserContext;
+};
+
 /// Contains information about the state of the editor that the assistant is currently docked with.
 USTRUCT(BlueprintType)
 struct FAIAssistantDockContext
 {
 	GENERATED_USTRUCT_BODY()
-
 public:
 	/// The asset instance that is currently being edited (if any).
 	UPROPERTY(BlueprintReadOnly, Category = "AIAssistant")
@@ -63,7 +81,7 @@ class UAIAssistantToolset : public UToolsetDefinition
 public:
 	/// Returns important information about the current project and/or current user.
 	UFUNCTION(BlueprintCallable, Category = "AI Assistant")
-	static FString GetProjectContext();
+	static FAIAssistantContext GetProjectContext();
 
 	/// Returns information about the part of the unreal UI that the assistant is docked into.
 	/// For example, if the assistant is docked with a Blueprint asset editor, this returns
diff --git a/Engine/Plugins/Experimental/AIAssistant/Source/AIAssistant/Private/Tests/AIAssistantToolsetTest.cpp b/Engine/Plugins/Experimental/AIAssistant/Source/AIAssistant/Private/Tests/AIAssistantToolsetTest.cpp
--- a/Engine/Plugins/Experimental/AIAssistant/Source/AIAssistant/Private/Tests/AIAssistantToolsetTest.cpp
+++ b/Engine/Plugins/Experimental/AIAssistant/Source/AIAssistant/Private/Tests/AIAssistantToolsetTest.cpp
@@ -45,13 +45,20 @@ void FAssistantToolsetSpec::Define()
 		UserPrompt.Empty();
 	});
 
+	It("Returns engine context.", [this]()
+	{
+		FAIAssistantContext Context = UAIAssistantToolset::GetProjectContext();
+		
+		TestTrue(TEXT("Context"), Context.UnrealContext.Contains(TEXT("100 units is 1 meter")));
+	});
+
 	It("Returns project context.", [this]()
 	{
 		UAIAssistantContextProject* ProjectSettings = GetMutableDefault<UAIAssistantContextProject>();
 		ProjectSettings->Prompt = FString(TEXT("Foo bar"));
 
-		FString Context = UAIAssistantToolset::GetProjectContext();
-		TestEqual(TEXT("Context"), Context, FString(TEXT("Foo bar")));
+		FAIAssistantContext Context = UAIAssistantToolset::GetProjectContext();
+		TestEqual(TEXT("Context"), Context.ProjectContext, FString(TEXT("Foo bar")));
 	});
 
 	It("Returns user context.", [this]()
@@ -59,21 +66,9 @@ void FAssistantToolsetSpec::Define()
 		UAIAssistantContextUser* UserSettings = GetMutableDefault<UAIAssistantContextUser>();
 		UserSettings->Prompt = FString(TEXT("Bish bosh"));
 
-		FString Context = UAIAssistantToolset::GetProjectContext();
-		TestEqual(TEXT("Context"), Context, FString(TEXT("Bish bosh")));
-	});
-
-	It("Returns project and user context.", [this]()
-	{
-		UAIAssistantContextProject* ProjectSettings = GetMutableDefault<UAIAssistantContextProject>();
-		ProjectSettings->Prompt = FString(TEXT("Foo bar"));
-
-		UAIAssistantContextUser* UserSettings = GetMutableDefault<UAIAssistantContextUser>();
-		UserSettings->Prompt = FString(TEXT("Bish bosh"));
-
-		FString Context = UAIAssistantToolset::GetProjectContext();
-		TestEqual(TEXT("Context"), Context, FString(TEXT("Foo bar\nBish bosh")));
-	});
+		FAIAssistantContext Context = UAIAssistantToolset::GetProjectContext();
+		TestEqual(TEXT("Context"), Context.UserContext, FString(TEXT("Bish bosh")));
+	});	
 }
 
 
