[AI Assistant] Backwards compatibility with an older stream.

This makes it possible to compile the assistant plugins with an older snapshot
of the engine.

Also, ensured AI assistant plugin setup is unregistered from OnPostEngineInit
when the plugin is unloaded.

#Synced-CL 51484694

diff --git a/Engine/Plugins/Experimental/AIAssistant/Source/AIAssistant/Private/AIAssistant.cpp b/Engine/Plugins/Experimental/AIAssistant/Source/AIAssistant/Private/AIAssistant.cpp
--- a/Engine/Plugins/Experimental/AIAssistant/Source/AIAssistant/Private/AIAssistant.cpp
+++ b/Engine/Plugins/Experimental/AIAssistant/Source/AIAssistant/Private/AIAssistant.cpp
@@ -1,6 +1,5 @@
 // Copyright Epic Games, Inc. All Rights Reserved.
 
-
 #include "AIAssistant.h"
 
 #include "Editor.h"
@@ -12,6 +11,7 @@
 #include "Slate/SceneViewport.h"
 #include "ToolMenus.h"
 #include "Widgets/Docking/SDockTab.h"
+#include "Templates/UnrealTemplate.h"
 
 #include "AIAssistantCommands.h"
 #include "AIAssistantCurrentConfig.h"
@@ -36,6 +36,23 @@ static const FName AIAssistantTabName("AIAssistant");
 // FAIAssistantModule
 //
 
+// Temporarily used to determine whether FCoreDelegates::GetOnPostEngineInit() is
+// available until it's present in all streams.
+template<typename T>
+struct FOnPostEngineInit
+{
+	static FSimpleMulticastDelegate& Get()
+	{
+		if constexpr (requires { T::GetOnPostEngineInit(); })
+		{
+			return T::GetOnPostEngineInit();
+		}
+		else
+		{
+			return T::OnPostEngineInit;
+		}
+	}
+};
 
 /*virtual*/ void FAIAssistantModule::StartupModule()
 {
@@ -78,18 +95,20 @@ static const FName AIAssistantTabName("AIAssistant");
 
 	InputProcessor = MakeShared<FAIAssistantInputProcessor>(PluginCommands);
 
+	auto SetupWhenEditorIsReady = [this]()
+		{
+			RegisterStatusBarPanelDrawerSummon();
+			UE::AIAssistant::FFileLockManager::RegisterWithEditor();
+		};
 	if (GEditor)
 	{
-		RegisterStatusBarPanelDrawerSummon();
-		UE::AIAssistant::FFileLockManager::RegisterWithEditor();
+		SetupWhenEditorIsReady();
 	}
 	else
 	{
-		FCoreDelegates::GetOnPostEngineInit().AddRaw(this, &FAIAssistantModule::RegisterStatusBarPanelDrawerSummon);
-		FCoreDelegates::GetOnPostEngineInit().AddLambda([]()
-		{
-			UE::AIAssistant::FFileLockManager::RegisterWithEditor();
-		});
+		FSimpleMulticastDelegate& OnPostEngineInit = FOnPostEngineInit<FCoreDelegates>::Get();
+		OnPostEngineInitHandle = UE::ToolsetRegistry::FDelegateHandleRaii::Create(
+			OnPostEngineInit, OnPostEngineInit.AddLambda(MoveTemp(SetupWhenEditorIsReady)));
 	}
 
 	// Some scenarios, like commandlets, don't have slate initialized
diff --git a/Engine/Plugins/Experimental/AIAssistant/Source/AIAssistant/Private/AIAssistant.h b/Engine/Plugins/Experimental/AIAssistant/Source/AIAssistant/Private/AIAssistant.h
--- a/Engine/Plugins/Experimental/AIAssistant/Source/AIAssistant/Private/AIAssistant.h
+++ b/Engine/Plugins/Experimental/AIAssistant/Source/AIAssistant/Private/AIAssistant.h
@@ -11,6 +11,7 @@
 
 
 #include "AIAssistantCurrentConfig.h"
+#include "ToolsetRegistry/DelegateHandle.h"
 
 class FAIAssistantInputProcessor;
 class FSpawnTabArgs;
@@ -93,4 +94,5 @@ private:
 	TSharedPtr<SAIAssistantWebBrowser> AIAssistantWebBrowserWidget;
 
 	FDelegateHandle StatusBarPanelDrawerSummonHandle;
+	UE::ToolsetRegistry::FDelegateHandleRaii OnPostEngineInitHandle;
 };
diff --git a/Engine/Plugins/Experimental/ToolsetRegistry/Source/ToolsetRegistry/Private/ToolsetRegistry/JsonConversion.cpp b/Engine/Plugins/Experimental/ToolsetRegistry/Source/ToolsetRegistry/Private/ToolsetRegistry/JsonConversion.cpp
--- a/Engine/Plugins/Experimental/ToolsetRegistry/Source/ToolsetRegistry/Private/ToolsetRegistry/JsonConversion.cpp
+++ b/Engine/Plugins/Experimental/ToolsetRegistry/Source/ToolsetRegistry/Private/ToolsetRegistry/JsonConversion.cpp
@@ -39,13 +39,13 @@ namespace UE::ToolsetRegistry::Internal
 		TSharedPtr<FJsonValue> ReturnJsonValue;
 		if (!JsonString.IsEmpty())
 		{
-			static const UE::FSharedString Key = TEXT("__key");
+			static const FString Key = TEXT("__key");
 			TSharedPtr<FJsonObject> JsonObject =
 				JsonStringToJsonObject(
 					FString::Printf(TEXT(R"json({"%s": %s})json"), *Key, *JsonString));
 			if (JsonObject)
 			{
-				ReturnJsonValue = JsonObject->Values.FindAndRemoveChecked(Key);
+				ReturnJsonValue = JsonObject->Values.FindAndRemoveChecked(*Key);
 			}
 		}
 		return ReturnJsonValue;
diff --git a/Engine/Plugins/Experimental/ToolsetRegistry/Source/ToolsetRegistry/Private/ToolsetRegistry/ObjectFunctionToolCall.cpp b/Engine/Plugins/Experimental/ToolsetRegistry/Source/ToolsetRegistry/Private/ToolsetRegistry/ObjectFunctionToolCall.cpp
--- a/Engine/Plugins/Experimental/ToolsetRegistry/Source/ToolsetRegistry/Private/ToolsetRegistry/ObjectFunctionToolCall.cpp
+++ b/Engine/Plugins/Experimental/ToolsetRegistry/Source/ToolsetRegistry/Private/ToolsetRegistry/ObjectFunctionToolCall.cpp
@@ -194,10 +194,9 @@ namespace UE::ToolsetRegistry
 				FunctionSchemaInputSchemaJsonObject->TryGetObjectField(
 					TEXT("properties"), InputPropertiesJsonObject))
 			{
-				for (const TPair<UE::FSharedString, TSharedPtr<FJsonValue>>& Pair :
-					(*InputPropertiesJsonObject)->Values)
+				for (const auto& Pair : (*InputPropertiesJsonObject)->Values)
 				{
-					const UE::FSharedString& InputParamPropertyName = Pair.Key;
+					const auto& InputParamPropertyName = Pair.Key;
 
 					const TSharedPtr<FJsonObject>& InputPropertyJsonObject =
 						Pair.Value->AsObject();
diff --git a/Engine/Plugins/Experimental/ToolsetRegistry/Source/ToolsetRegistry/Private/ToolsetRegistry/ToolsetLibrary.cpp b/Engine/Plugins/Experimental/ToolsetRegistry/Source/ToolsetRegistry/Private/ToolsetRegistry/ToolsetLibrary.cpp
--- a/Engine/Plugins/Experimental/ToolsetRegistry/Source/ToolsetRegistry/Private/ToolsetRegistry/ToolsetLibrary.cpp
+++ b/Engine/Plugins/Experimental/ToolsetRegistry/Source/ToolsetRegistry/Private/ToolsetRegistry/ToolsetLibrary.cpp
@@ -5,12 +5,9 @@
 #include "AssetRegistry/AssetRegistryHelpers.h"
 #include "AssetRegistry/AssetRegistryModule.h"
 #include "Dom/JsonObject.h"
-#include "LevelEditorViewport.h"
 #include "Misc/AssertionMacros.h"
 #include "UObject/PropertyAccessUtil.h"
-#include "Serialization/JsonSerializer.h"
-#include "Serialization/JsonReader.h"
-#include "Serialization/JsonWriter.h"
+#include "UnrealClient.h"
 
 #include "ToolsetRegistry/JsonConversion.h"
 #include "ToolsetRegistry/Module.h"
@@ -69,10 +66,9 @@ bool UToolsetLibrary::SetObjectProperties(UObject* Object, const FString& Proper
 	{
 		AllPropertiesSet = true;
 		const UClass* ObjectClass = Object->GetClass();
-		TSet<UE::FSharedString> Keys;
-		PropertyData->Values.GetKeys(Keys);
-		for (const auto& Key : Keys)
+		for (const auto& Pair : PropertyData->Values)
 		{
+			const auto& Key = Pair.Key;
 			FProperty* Property = ObjectClass->FindPropertyByName(FName(Key));
 			if (Property)
 			{
