[AI Assistant] Use backward compatible template to get editor delegates.

#Synced-CL 51547877

diff --git a/Engine/Plugins/Experimental/AIAssistant/Source/AIAssistant/Private/AIAssistant.cpp b/Engine/Plugins/Experimental/AIAssistant/Source/AIAssistant/Private/AIAssistant.cpp
--- a/Engine/Plugins/Experimental/AIAssistant/Source/AIAssistant/Private/AIAssistant.cpp
+++ b/Engine/Plugins/Experimental/AIAssistant/Source/AIAssistant/Private/AIAssistant.cpp
@@ -20,6 +20,7 @@
 #include "AIAssistantStyle.h"
 #include "AIAssistantWebBrowser.h"
 #include "AIAssistantFileLockManager.h"
+#include "ToolsetRegistry/Private/ToolsetRegistry/EngineDelegatesCompatibility.h"
 
 #define LOCTEXT_NAMESPACE "FAIAssistantModule"
 
@@ -36,24 +37,6 @@ static const FName AIAssistantTabName("AIAssistant");
 // FAIAssistantModule
 //
 
-// Temporarily used to determine whether FCoreDelegates::GetOnPostEngineInit() is
-// available until it's present in all streams.
-template<typename T>
-struct FOnPostEngineInit
-{
-	static FSimpleMulticastDelegate& Get()
-	{
-		if constexpr (requires { T::GetOnPostEngineInit(); })
-		{
-			return T::GetOnPostEngineInit();
-		}
-		else
-		{
-			return T::OnPostEngineInit;
-		}
-	}
-};
-
 /*virtual*/ void FAIAssistantModule::StartupModule()
 {
 	CurrentConfig = MakeShared<UE::AIAssistant::FCurrentConfig>();
diff --git a/Engine/Plugins/Experimental/ToolsetRegistry/Source/ToolsetRegistry/Private/ToolsetRegistry/EngineDelegatesCompatibility.h b/Engine/Plugins/Experimental/ToolsetRegistry/Source/ToolsetRegistry/Private/ToolsetRegistry/EngineDelegatesCompatibility.h
new file mode 100644
--- /dev/null
+++ b/Engine/Plugins/Experimental/ToolsetRegistry/Source/ToolsetRegistry/Private/ToolsetRegistry/EngineDelegatesCompatibility.h
@@ -0,0 +1,23 @@
+// Copyright Epic Games, Inc. All Rights Reserved.
+
+#pragma once
+
+#include "Delegates/Delegate.h"
+
+// Temporarily used to determine whether FCoreDelegates::GetOnPostEngineInit() is
+// available until it's present in all streams.
+template<typename T>
+struct FOnPostEngineInit
+{
+	static FSimpleMulticastDelegate& Get()
+	{
+		if constexpr (requires { T::GetOnPostEngineInit(); })
+		{
+			return T::GetOnPostEngineInit();
+		}
+		else
+		{
+			return T::OnPostEngineInit;
+		}
+	}
+};
\ No newline at end of file
