Remove GraphToText tool as we have better support through domain specific toolsets.

#Synced-CL 51513449

diff --git a/Engine/Plugins/Experimental/AIAssistant/Source/AIAssistant/Private/AIAssistantToolset.cpp b/Engine/Plugins/Experimental/AIAssistant/Source/AIAssistant/Private/AIAssistantToolset.cpp
--- a/Engine/Plugins/Experimental/AIAssistant/Source/AIAssistant/Private/AIAssistantToolset.cpp
+++ b/Engine/Plugins/Experimental/AIAssistant/Source/AIAssistant/Private/AIAssistantToolset.cpp
@@ -133,20 +133,3 @@ FAIAssistantDockContext UAIAssistantToolset::GetDockedContext()
 
 	return FAIAssistantDockContext();
 }
-
-FString UAIAssistantToolset::GraphToText(const FAIAssistantDockContext& DockContext)
-{
-	FString ExportedText;
-	if (DockContext.Graph)
-	{
-		TArray<TObjectPtr<UEdGraphNode>> SelectedNodes = DockContext.SelectedNodes.Num() > 0 ? DockContext.SelectedNodes : DockContext.Graph->Nodes;
-
-		TSet<UObject*> NodesToExport;
-		for (const auto& Node : SelectedNodes)
-		{
-			NodesToExport.Add(Node.Get());
-		}		
-		FEdGraphUtilities::ExportNodesToText(NodesToExport, ExportedText);
-	}
-	return ExportedText;
-}
diff --git a/Engine/Plugins/Experimental/AIAssistant/Source/AIAssistant/Private/AIAssistantToolset.h b/Engine/Plugins/Experimental/AIAssistant/Source/AIAssistant/Private/AIAssistantToolset.h
--- a/Engine/Plugins/Experimental/AIAssistant/Source/AIAssistant/Private/AIAssistantToolset.h
+++ b/Engine/Plugins/Experimental/AIAssistant/Source/AIAssistant/Private/AIAssistantToolset.h
@@ -89,9 +89,4 @@ public:
 	/// is closed or not docked to an asset editor, this information will be empty.
 	UFUNCTION(BlueprintCallable, Category = "AI Assistant")
 	static FAIAssistantDockContext GetDockedContext();
-
-	/// This converts graph information in a dock context into a string that describes it.
-	/// This string will represent the entire graph if SelectedNodes is empty.
-	UFUNCTION(BlueprintCallable, Category = "AI Assistant")
-	static FString GraphToText(const FAIAssistantDockContext& DockContext);
 };
