Fix ASan issues with AIAssistantWebApplicationTest

#Synced-CL 52358950

diff --git a/Engine/Plugins/Experimental/AIAssistant/Source/AIAssistant/Private/Tests/AIAssistantWebApplicationTest.cpp b/Engine/Plugins/Experimental/AIAssistant/Source/AIAssistant/Private/Tests/AIAssistantWebApplicationTest.cpp
--- a/Engine/Plugins/Experimental/AIAssistant/Source/AIAssistant/Private/Tests/AIAssistantWebApplicationTest.cpp
+++ b/Engine/Plugins/Experimental/AIAssistant/Source/AIAssistant/Private/Tests/AIAssistantWebApplicationTest.cpp
@@ -1852,7 +1852,10 @@ namespace UE::AIAssistant
 				FakeWebApi->RegisteredCallbackHandlerIds[0], EventJson, false);
 		}
 
-	private:
+	private: 
+		// FSpyWebApplication cannot be the last to hold  the ToolsetRegistry SharedPtr 
+		// because its base class holds a reference to it (via the OnToolsetRegisteredCallbackHandle)
+		TSharedPtr<UE::ToolsetRegistry::FToolsetRegistry> ToolsetRegistry;
 		TOptional<FSpyWebApplicationWithTracker> WebApplicationWithTracker;
 		TSharedPtr<FSpyWebApplication> WebApplication;
 		TSharedPtr<FFakeWebApi> WebApi;
@@ -1863,10 +1866,10 @@ namespace UE::AIAssistant
 	{
 		BeforeEach([this]()
 		{
+			ToolsetRegistry = MakeShared<UE::ToolsetRegistry::FToolsetRegistry>();
 			WebApplicationWithTracker.Emplace(*this);
 			WebApplication = WebApplicationWithTracker->WebApplication;
-			WebApplication->SetRegistry(
-				MakeShared<UE::ToolsetRegistry::FToolsetRegistry>());
+			WebApplication->SetRegistry(ToolsetRegistry);
 			if (!WebApplicationWithTracker->EnsureInitialized()) return;
 			WebApi = WebApplicationWithTracker->GetWebApi();
 			TestTrue(TEXT("HasWebApi"), WebApi.IsValid());
@@ -1877,6 +1880,7 @@ namespace UE::AIAssistant
 			WebApi.Reset();
 			WebApplication.Reset();
 			WebApplicationWithTracker.Reset();
+			ToolsetRegistry.Reset();
 		});
 
 		Describe(TEXT("when a generationStopped event arrives"), [this]
